[NO NON-PII COMMIT MESSAGE]

diff --git a/Engine/Plugins/Runtime/ResonanceAudio/Source/ResonanceAudio/Private/ResonanceAudioCommon.cpp b/Engine/Plugins/Runtime/ResonanceAudio/Source/ResonanceAudio/Private/ResonanceAudioCommon.cpp
--- a/Engine/Plugins/Runtime/ResonanceAudio/Source/ResonanceAudio/Private/ResonanceAudioCommon.cpp
+++ b/Engine/Plugins/Runtime/ResonanceAudio/Source/ResonanceAudio/Private/ResonanceAudioCommon.cpp
@@ -47,12 +47,12 @@ namespace ResonanceAudio
 		}
 		else
 		{
-			UE_LOG(LogResonanceAudio, Error, TEXT("File does not exist. %s"), *DynamicLibraryToLoad);
+			UE_LOG(LogResonanceAudio, Log, TEXT("File does not exist. %s"), *DynamicLibraryToLoad);
 		}
 
 		if (!DynamicLibraryHandle)
 		{
-			UE_LOG(LogResonanceAudio, Error, TEXT("Unable to load %s."), *FPaths::ConvertRelativePathToFull(DynamicLibraryToLoad));
+			UE_LOG(LogResonanceAudio, Log, TEXT("Unable to load %s."), *FPaths::ConvertRelativePathToFull(DynamicLibraryToLoad));
 		}
 		else
 		{
@@ -65,14 +65,21 @@ namespace ResonanceAudio
 	vraudio::VrAudioApi* CreateResonanceAudioApi(void* DynamicLibraryHandle, size_t NumChannels, size_t NumFrames, int SampleRate) {
 		vraudio::VrAudioApi* (*create)(size_t, size_t, int);
 #if PLATFORM_LINUX || PLATFORM_MAC || PLATFORM_WINDOWS
-		 create = reinterpret_cast<vraudio::VrAudioApi* (*)(size_t, size_t, int)>(FPlatformProcess::GetDllExport(DynamicLibraryHandle, TEXT("CreateVrAudioApi")));
+		if (DynamicLibraryHandle)
+		{
+			create = reinterpret_cast<vraudio::VrAudioApi* (*)(size_t, size_t, int)>(FPlatformProcess::GetDllExport(DynamicLibraryHandle, TEXT("CreateVrAudioApi")));
+		}
+		else
+		{
+			create = nullptr;
+		}
 #else
 		 // For the static case, or for Android.
 		return vraudio::CreateVrAudioApi(NumChannels, NumFrames, SampleRate);
 #endif
 
 		if (create == nullptr) {
-			UE_LOG(LogResonanceAudio, Error, TEXT("Failed to load the Create method from VrAudioApi."));
+			UE_LOG(LogResonanceAudio, Log, TEXT("Failed to load the Create method from VrAudioApi."));
 			return nullptr;
 		}
 		return create(NumChannels, NumFrames, SampleRate);
