Adding tweaks to google resonance plugin

- Not require content to be checked in to use
- New cvar to toggle enabling binaural spat on all sources (where available) so we can set in project and/or live

diff --git a/Engine/Plugins/Runtime/ResonanceAudio/Source/ResonanceAudio/Private/ResonanceAudioModule.cpp b/Engine/Plugins/Runtime/ResonanceAudio/Source/ResonanceAudio/Private/ResonanceAudioModule.cpp
--- a/Engine/Plugins/Runtime/ResonanceAudio/Source/ResonanceAudio/Private/ResonanceAudioModule.cpp
+++ b/Engine/Plugins/Runtime/ResonanceAudio/Source/ResonanceAudio/Private/ResonanceAudioModule.cpp
@@ -8,12 +8,14 @@
 #include "ResonanceAudioReverb.h"
 #include "ResonanceAudioSpatialization.h"
 #include "Features/IModularFeatures.h"
+#include "ResonanceAudioSettings.h"
 
 IMPLEMENT_MODULE(ResonanceAudio::FResonanceAudioModule, ResonanceAudio)
 
 namespace ResonanceAudio
 {
 	void* FResonanceAudioModule::ResonanceAudioDynamicLibraryHandle = nullptr;
+	UResonanceAudioSpatializationSourceSettings* FResonanceAudioModule::GlobalSpatializationSourceSettings = nullptr;
 
 	static bool bModuleInitialized = false;
 
@@ -39,10 +41,35 @@ namespace ResonanceAudio
 		{
 			ResonanceAudioDynamicLibraryHandle = LoadResonanceAudioDynamicLibrary();
 		}
+
+		if (!IsRunningDedicatedServer() && !GlobalSpatializationSourceSettings)
+		{
+			// Load the global source preset settings:
+			const FSoftObjectPath GlobalPluginPresetName = GetDefault<UResonanceAudioSettings>()->GlobalSourcePreset;
+			if (GlobalPluginPresetName.IsValid())
+			{
+				GlobalSpatializationSourceSettings = LoadObject<UResonanceAudioSpatializationSourceSettings>(nullptr, *GlobalPluginPresetName.ToString());
+			}
+			else
+			{
+				GlobalSpatializationSourceSettings = NewObject<UResonanceAudioSpatializationSourceSettings>(UResonanceAudioSpatializationSourceSettings::StaticClass(), TEXT("Default Global Resonance Spatialization Preset"));
+			}
+
+			if (GlobalSpatializationSourceSettings)
+			{ 
+				GlobalSpatializationSourceSettings->AddToRoot();
+			}
+		}
 	}
 
 	void FResonanceAudioModule::ShutdownModule()
 	{
+		if (GlobalSpatializationSourceSettings)
+		{
+			GlobalSpatializationSourceSettings->RemoveFromRoot();
+			GlobalSpatializationSourceSettings = nullptr;
+		}
+
 		check(bModuleInitialized == true);
 		bModuleInitialized = false;
 	}
@@ -77,6 +104,11 @@ namespace ResonanceAudio
 		RegisteredAudioDevices.Remove(AudioDeviceHandle);
 	}
 
+	UResonanceAudioSpatializationSourceSettings* FResonanceAudioModule::GetGlobalSpatializationSourceSettings()
+	{
+		return GlobalSpatializationSourceSettings;
+	}
+
 	TAudioSpatializationPtr FSpatializationPluginFactory::CreateNewSpatializationPlugin(FAudioDevice* OwningDevice)
 	{
 		// Register the audio device to the Resonance Audio module.
diff --git a/Engine/Plugins/Runtime/ResonanceAudio/Source/ResonanceAudio/Private/ResonanceAudioModule.h b/Engine/Plugins/Runtime/ResonanceAudio/Source/ResonanceAudio/Private/ResonanceAudioModule.h
--- a/Engine/Plugins/Runtime/ResonanceAudio/Source/ResonanceAudio/Private/ResonanceAudioModule.h
+++ b/Engine/Plugins/Runtime/ResonanceAudio/Source/ResonanceAudio/Private/ResonanceAudioModule.h
@@ -65,6 +65,11 @@ namespace ResonanceAudio
 
 		virtual TAmbisonicsMixerPtr CreateNewAmbisonicsMixer(FAudioDevice* OwningDevice) override;
 
+		virtual int32 GetMaxSupportedChannels() override
+		{
+			return 2;
+		}
+
 	};
 
 	/******************************************************/
@@ -143,10 +148,16 @@ namespace ResonanceAudio
 		// Returns Resonance Audio API dynamic library handle.
 		static void* GetResonanceAudioDynamicLibraryHandle() { return ResonanceAudioDynamicLibraryHandle; };
 
+		// Returns the global spatialization source settings set in the project settings
+		static UResonanceAudioSpatializationSourceSettings* GetGlobalSpatializationSourceSettings();
+
 	private:
 		// Resonance Audio API dynamic library handle.
 		static void* ResonanceAudioDynamicLibraryHandle;
 
+		// The global resonance audio spatialization source settings
+		static UResonanceAudioSpatializationSourceSettings* GlobalSpatializationSourceSettings;
+
 		// List of registered audio devices.
 		TArray<FAudioDevice*> RegisteredAudioDevices;
 
diff --git a/Engine/Plugins/Runtime/ResonanceAudio/Source/ResonanceAudio/Private/ResonanceAudioSettings.h b/Engine/Plugins/Runtime/ResonanceAudio/Source/ResonanceAudio/Private/ResonanceAudioSettings.h
--- a/Engine/Plugins/Runtime/ResonanceAudio/Source/ResonanceAudio/Private/ResonanceAudioSettings.h
+++ b/Engine/Plugins/Runtime/ResonanceAudio/Source/ResonanceAudio/Private/ResonanceAudioSettings.h
@@ -23,4 +23,8 @@ public:
 	// Default settings for global reverb: This is overridden when a player enters Audio Volumes.
 	UPROPERTY(GlobalConfig, EditAnywhere, Category = General, meta = (AllowedClasses = "ResonanceAudioReverbPluginPreset"))
 	FSoftObjectPath GlobalReverbPreset;
+
+	// Default settings for global source settings
+	UPROPERTY(GlobalConfig, EditAnywhere, Category = General, meta = (AllowedClasses = "ResonanceAudioSourcePluginPreset"))
+	FSoftObjectPath GlobalSourcePreset;
 };
diff --git a/Engine/Plugins/Runtime/ResonanceAudio/Source/ResonanceAudio/Private/ResonanceAudioSpatialization.cpp b/Engine/Plugins/Runtime/ResonanceAudio/Source/ResonanceAudio/Private/ResonanceAudioSpatialization.cpp
--- a/Engine/Plugins/Runtime/ResonanceAudio/Source/ResonanceAudio/Private/ResonanceAudioSpatialization.cpp
+++ b/Engine/Plugins/Runtime/ResonanceAudio/Source/ResonanceAudio/Private/ResonanceAudioSpatialization.cpp
@@ -55,7 +55,16 @@ namespace ResonanceAudio {
 			BinauralSource.Id = RA_INVALID_SOURCE_ID;
 		}
 
-		SpatializationSettings[SourceId] = static_cast<UResonanceAudioSpatializationSourceSettings*>(InSettings);
+		// If we weren't passed in any settings, we want to get the settings from the global project settings
+		if (!InSettings)
+		{
+			SpatializationSettings[SourceId] = FResonanceAudioModule::GetGlobalSpatializationSourceSettings();
+		}
+		else
+		{
+			SpatializationSettings[SourceId] = static_cast<UResonanceAudioSpatializationSourceSettings*>(InSettings);
+		}
+
 
 		if (SpatializationSettings[SourceId] == nullptr)
 		{
diff --git a/Engine/Source/Runtime/Engine/Private/AudioDevice.cpp b/Engine/Source/Runtime/Engine/Private/AudioDevice.cpp
--- a/Engine/Source/Runtime/Engine/Private/AudioDevice.cpp
+++ b/Engine/Source/Runtime/Engine/Private/AudioDevice.cpp
@@ -135,6 +135,14 @@ FAutoConsoleVariableRef CVarSoundDistanceOptimizationLength(
 	TEXT("The maximum duration a sound must be in order to be a candidate to be culled due to one-shot distance optimization.\n"),
 	ECVF_Default);
 
+static int32 EnableBinauralAudioForAllSpatialSoundsCVar = 0;
+FAutoConsoleVariableRef CVarEnableBinauralAudioForAllSpatialSounds(
+	TEXT("au.EnableBinauralAudioForAllSpatialSounds"),
+	EnableBinauralAudioForAllSpatialSoundsCVar,
+	TEXT("Toggles binaural audio rendering for all spatial sounds if binaural rendering is available.\n"),
+	ECVF_Default);
+
+
 using FVirtualLoopPair = TPair<FActiveSound*, FAudioVirtualLoop>;
 
 /*-----------------------------------------------------------------------------
@@ -1650,11 +1658,11 @@ bool FAudioDevice::IsHRTFEnabledForAll() const
 {
 	if (IsInAudioThread())
 	{
-		return bHRTFEnabledForAll && IsSpatializationPluginEnabled();
+		return (bHRTFEnabledForAll || EnableBinauralAudioForAllSpatialSoundsCVar == 1) && IsSpatializationPluginEnabled();
 	}
 
 	check(IsInGameThread());
-	return bHRTFEnabledForAll_OnGameThread && IsSpatializationPluginEnabled();
+	return (bHRTFEnabledForAll_OnGameThread || EnableBinauralAudioForAllSpatialSoundsCVar == 1) && IsSpatializationPluginEnabled();
 }
 
 void FAudioDevice::SetMixDebugState(EDebugState InDebugState)
