Fix for compiler error in nonunity compiles. #jira none

diff --git a/Engine/Plugins/Runtime/ResonanceAudio/Source/ResonanceAudio/Private/NullEmbreeImplementation.cpp b/Engine/Plugins/Runtime/ResonanceAudio/Source/ResonanceAudio/Private/NullEmbreeImplementation.cpp
--- a/Engine/Plugins/Runtime/ResonanceAudio/Source/ResonanceAudio/Private/NullEmbreeImplementation.cpp
+++ b/Engine/Plugins/Runtime/ResonanceAudio/Source/ResonanceAudio/Private/NullEmbreeImplementation.cpp
@@ -7,6 +7,7 @@
 
 #include "embree2/rtcore.h"
 #include "embree2/rtcore_ray.h"
+#include "CoreMinimal.h"
 
 void rtcCommit(RTCScene scene)
 {
