Fix for static analysis warnings in Google Resonance.
#rnx

diff --git a/Engine/Plugins/Runtime/ResonanceAudio/Source/ResonanceAudio/Private/ResonanceAudioLibrary/resonance_audio/utils/wav_reader.cc b/Engine/Plugins/Runtime/ResonanceAudio/Source/ResonanceAudio/Private/ResonanceAudioLibrary/resonance_audio/utils/wav_reader.cc
--- a/Engine/Plugins/Runtime/ResonanceAudio/Source/ResonanceAudio/Private/ResonanceAudioLibrary/resonance_audio/utils/wav_reader.cc
+++ b/Engine/Plugins/Runtime/ResonanceAudio/Source/ResonanceAudio/Private/ResonanceAudioLibrary/resonance_audio/utils/wav_reader.cc
@@ -137,7 +137,7 @@ bool WavReader::ParseHeader() {
   sample_rate_hz_ = header.format.samples_rate;
 
   bytes_per_sample_ = header.format.bits_per_sample / 8;
-  if (bytes_per_sample_ == 0 || bytes_per_sample_ != sizeof(int16)) {
+  if (bytes_per_sample_ != sizeof(int16)) {
     return false;
   }
   const size_t bytes_in_payload = header.data.header.size;
