#rnx

diff --git a/Engine/Plugins/Runtime/ResonanceAudio/Source/ResonanceAudio/Private/ResonanceAudioLibrary/resonance_audio/base/constants_and_types.h b/Engine/Plugins/Runtime/ResonanceAudio/Source/ResonanceAudio/Private/ResonanceAudioLibrary/resonance_audio/base/constants_and_types.h
--- a/Engine/Plugins/Runtime/ResonanceAudio/Source/ResonanceAudio/Private/ResonanceAudioLibrary/resonance_audio/base/constants_and_types.h
+++ b/Engine/Plugins/Runtime/ResonanceAudio/Source/ResonanceAudio/Private/ResonanceAudioLibrary/resonance_audio/base/constants_and_types.h
@@ -181,6 +181,7 @@ static const float kRotationQuantizationRad = 1.0f * kRadiansFromDegrees;
 }  // namespace vraudio
 
 #ifdef DEFINED_M_PI
+#undef DEFINED_M_PI
 #undef M_PI
 #endif
 
