Fix compile error in ResonanceAudio, and remove change to disable optimization.

#rnx

diff --git a/Engine/Plugins/Runtime/ResonanceAudio/Source/ResonanceAudio/ResonanceAudio.Build.cs b/Engine/Plugins/Runtime/ResonanceAudio/Source/ResonanceAudio/ResonanceAudio.Build.cs
--- a/Engine/Plugins/Runtime/ResonanceAudio/Source/ResonanceAudio/ResonanceAudio.Build.cs
+++ b/Engine/Plugins/Runtime/ResonanceAudio/Source/ResonanceAudio/ResonanceAudio.Build.cs
@@ -8,9 +8,6 @@ namespace UnrealBuildTool.Rules
     {
         public ResonanceAudio(ReadOnlyTargetRules Target) : base(Target)
         {
-
-            OptimizeCode = CodeOptimization.Never;
-
             string ResonanceAudioPath = ModuleDirectory + "/Private/ResonanceAudioLibrary";
             string ResonanceAudioLibraryPath = ModuleDirectory + "/Private/ResonanceAudioLibrary/resonance_audio";
             string PFFTPath = ModuleDirectory + "/Private/ResonanceAudioLibrary/third_party/pfft";
@@ -108,6 +105,7 @@ namespace UnrealBuildTool.Rules
                 string SDKDir = Target.UEThirdPartySourceDirectory + "IntelEmbree/Embree2140/Win64/";
                 PublicIncludePaths.Add(SDKDir + "include");
                 PublicDefinitions.Add("USE_EMBREE=0");
+				PublicDefinitions.Add("EMBREE_STATIC_LIB=1");
             }
         }
     }
