optionally compile out suspicious looking call to abort on FatalLogger. #jira none.
#rnx

diff --git a/Engine/Plugins/Runtime/ResonanceAudio/Source/ResonanceAudio/Private/ResonanceAudioLibrary/resonance_audio/base/logging.h b/Engine/Plugins/Runtime/ResonanceAudio/Source/ResonanceAudio/Private/ResonanceAudioLibrary/resonance_audio/base/logging.h
--- a/Engine/Plugins/Runtime/ResonanceAudio/Source/ResonanceAudio/Private/ResonanceAudioLibrary/resonance_audio/base/logging.h
+++ b/Engine/Plugins/Runtime/ResonanceAudio/Source/ResonanceAudio/Private/ResonanceAudioLibrary/resonance_audio/base/logging.h
@@ -65,6 +65,8 @@ class NullLogger {
 #define DCHECK_GE(val1, val2) DCHECK_OP((val1), (val2), >=)
 #define DCHECK_GT(val1, val2) DCHECK_OP((val1), (val2), >)
 
+#define ABORT_ON_FATAL_LOG 0
+
 // This class is used to log to std::cerr.
 class FatalLogger {
  public:
@@ -74,7 +76,9 @@ class FatalLogger {
   ~FatalLogger() {
     const std::string error_string = error_string_.str();
     std::cerr << error_string << std::endl;
+#if ABORT_ON_FATAL_LOG
     abort();
+#endif
   }
   std::ostream& GetStream() { return error_string_; }
 
