[NO NON-PII COMMIT MESSAGE]

diff --git a/Engine/Plugins/Runtime/ResonanceAudio/Source/ResonanceAudio/Private/ResonanceAudioCommon.h b/Engine/Plugins/Runtime/ResonanceAudio/Source/ResonanceAudio/Private/ResonanceAudioCommon.h
--- a/Engine/Plugins/Runtime/ResonanceAudio/Source/ResonanceAudio/Private/ResonanceAudioCommon.h
+++ b/Engine/Plugins/Runtime/ResonanceAudio/Source/ResonanceAudio/Private/ResonanceAudioCommon.h
@@ -4,9 +4,11 @@
 
 #pragma once
 
+THIRD_PARTY_INCLUDES_START
 #include "api/resonance_audio_api.h"
 #include "platforms/common/room_properties.h"
 #include "platforms/common/room_effects_utils.h"
+THIRD_PARTY_INCLUDES_END
 
 #include "CoreMinimal.h"
 #include "UObject/ObjectMacros.h"
diff --git a/Engine/Plugins/Runtime/ResonanceAudio/Source/ResonanceAudio/Private/ResonanceAudioLibrary/resonance_audio/ambisonics/hoa_rotator.h b/Engine/Plugins/Runtime/ResonanceAudio/Source/ResonanceAudio/Private/ResonanceAudioLibrary/resonance_audio/ambisonics/hoa_rotator.h
--- a/Engine/Plugins/Runtime/ResonanceAudio/Source/ResonanceAudio/Private/ResonanceAudioLibrary/resonance_audio/ambisonics/hoa_rotator.h
+++ b/Engine/Plugins/Runtime/ResonanceAudio/Source/ResonanceAudio/Private/ResonanceAudioLibrary/resonance_audio/ambisonics/hoa_rotator.h
@@ -19,7 +19,12 @@ limitations under the License.
 
 #include <vector>
 
+#include "CoreMinimal.h"
+
+THIRD_PARTY_INCLUDES_START
 #include "Eigen/Dense"
+THIRD_PARTY_INCLUDES_END
+
 #include "base/audio_buffer.h"
 #include "base/misc_math.h"
 
diff --git a/Engine/Plugins/Runtime/ResonanceAudio/Source/ResonanceAudio/Private/ResonanceAudioLibrary/resonance_audio/base/logging.h b/Engine/Plugins/Runtime/ResonanceAudio/Source/ResonanceAudio/Private/ResonanceAudioLibrary/resonance_audio/base/logging.h
--- a/Engine/Plugins/Runtime/ResonanceAudio/Source/ResonanceAudio/Private/ResonanceAudioLibrary/resonance_audio/base/logging.h
+++ b/Engine/Plugins/Runtime/ResonanceAudio/Source/ResonanceAudio/Private/ResonanceAudioLibrary/resonance_audio/base/logging.h
@@ -50,12 +50,20 @@ class NullLogger {
   }
 };
 
+
+#define DISABLE_CHECKS 1
+
+#if DISABLE_CHECKS
+#define DCHECK(expr) 
+#else
 // If statement prevents unused variable warnings.
 #define DCHECK(expr)   \
   if (false && (expr)) \
     ;                  \
   else                 \
     NullLogger().GetStream()
+#endif
+
 #define DCHECK_OP(val1, val2, op) DCHECK((val1)op(val2))
 
 #define DCHECK_EQ(val1, val2) DCHECK_OP((val1), (val2), ==)
diff --git a/Engine/Plugins/Runtime/ResonanceAudio/Source/ResonanceAudio/Private/ResonanceAudioLibrary/resonance_audio/base/misc_math.h b/Engine/Plugins/Runtime/ResonanceAudio/Source/ResonanceAudio/Private/ResonanceAudioLibrary/resonance_audio/base/misc_math.h
--- a/Engine/Plugins/Runtime/ResonanceAudio/Source/ResonanceAudio/Private/ResonanceAudioLibrary/resonance_audio/base/misc_math.h
+++ b/Engine/Plugins/Runtime/ResonanceAudio/Source/ResonanceAudio/Private/ResonanceAudioLibrary/resonance_audio/base/misc_math.h
@@ -29,7 +29,9 @@ limitations under the License.
 #include <vector>
 
 #include "base/integral_types.h"
+THIRD_PARTY_INCLUDES_START
 #include "Eigen/Dense"
+THIRD_PARTY_INCLUDES_END
 #include "base/constants_and_types.h"
 #include "base/logging.h"
 
diff --git a/Engine/Plugins/Runtime/ResonanceAudio/Source/ResonanceAudio/ResonanceAudio.Build.cs b/Engine/Plugins/Runtime/ResonanceAudio/Source/ResonanceAudio/ResonanceAudio.Build.cs
--- a/Engine/Plugins/Runtime/ResonanceAudio/Source/ResonanceAudio/ResonanceAudio.Build.cs
+++ b/Engine/Plugins/Runtime/ResonanceAudio/Source/ResonanceAudio/ResonanceAudio.Build.cs
@@ -56,9 +56,6 @@ namespace UnrealBuildTool.Rules
                 }
             );
 
-            // Resonance's Eigen usage is only valid with Eigen 2, so we set Eigen to backwards compatibility mode:
-            // PublicDefinitions.Add("EIGEN2_SUPPORT=1");
-
             if (Target.bBuildEditor == true)
             {
                 PrivateDependencyModuleNames.Add("UnrealEd");
