Add macro to avoid compiling Procedural Mesh in Resonance Audio plugin.
don't whitelist HTML5 for the Resonance Audio plugin for now.

diff --git a/Engine/Plugins/Runtime/ResonanceAudio/ResonanceAudio.uplugin b/Engine/Plugins/Runtime/ResonanceAudio/ResonanceAudio.uplugin
--- a/Engine/Plugins/Runtime/ResonanceAudio/ResonanceAudio.uplugin
+++ b/Engine/Plugins/Runtime/ResonanceAudio/ResonanceAudio.uplugin
@@ -15,11 +15,22 @@
   "IsBetaVersion": true,
   "Installed": false,
   "Modules": [
-    {
-      "Name": "ResonanceAudio",
-      "Type": "Runtime",
-      "LoadingPhase": "Default"
-    },
+	{
+		"Name": "ResonanceAudio",
+		"Type": "Runtime",
+		"LoadingPhase": "Default",
+		"WhitelistPlatforms": [
+			"Android",
+			"iOS",
+			"Linux",
+			"Mac",
+			"Win32",
+			"Win64",
+			"Lumin",
+			"XboxOne",
+			"PS4"
+		]
+	},
     {
       "Name": "ResonanceAudioEditor",
       "Type": "Editor",
diff --git a/Engine/Plugins/Runtime/ResonanceAudio/Source/ResonanceAudio/Private/ResonanceAudioDirectivityVisualizer.cpp b/Engine/Plugins/Runtime/ResonanceAudio/Source/ResonanceAudio/Private/ResonanceAudioDirectivityVisualizer.cpp
--- a/Engine/Plugins/Runtime/ResonanceAudio/Source/ResonanceAudio/Private/ResonanceAudioDirectivityVisualizer.cpp
+++ b/Engine/Plugins/Runtime/ResonanceAudio/Source/ResonanceAudio/Private/ResonanceAudioDirectivityVisualizer.cpp
@@ -12,18 +12,16 @@ AResonanceAudioDirectivityVisualizer::AResonanceAudioDirectivityVisualizer()
 {
 	// Make sure visualization is not displayed in the actual game.
 	Super::SetActorHiddenInGame(true);
+
+#if SUPPORTS_PROCEDURAL_MESH
 	Mesh = CreateDefaultSubobject<UProceduralMeshComponent>(TEXT("VisualizationMesh"));
+	Mesh->AddToRoot();
 
 	// Disable collision data.
 	Mesh->ContainsPhysicsTriMeshData(false);
 
 	RootComponent = Mesh;
-
-// 	static ConstructorHelpers::FObjectFinder<UMaterial> FoundMaterial(TEXT("Material'/ResonanceAudio/VisualizationMeshMaterial.VisualizationMeshMaterial'"));
-// 	if (FoundMaterial.Succeeded())
-// 	{
-// 		Material = FoundMaterial.Object;
-// 	}
+#endif // SUPPORTS_PROCEDURAL_MESH
 }
 
 void AResonanceAudioDirectivityVisualizer::DrawPattern()
@@ -57,6 +55,19 @@ void AResonanceAudioDirectivityVisualizer::DrawPattern()
 		}
 	}
 
+#if SUPPORTS_PROCEDURAL_MESH
 	Mesh->CreateMeshSection(0, Vertices, Triangles, Normals, UV0, VertexColors, Tangents, false);
 	Mesh->SetMaterial(0, Material);
+#endif // SUPPORTS_PROCEDURAL_MESH
+}
+
+void AResonanceAudioDirectivityVisualizer::BeginDestroy()
+{
+#if SUPPORTS_PROCEDURAL_MESH
+	if (Mesh)
+	{
+		Mesh->RemoveFromRoot();
+	}
+#endif // SUPPORTS_PROCEDURAL_MESH
+	AActor::BeginDestroy();
 }
diff --git a/Engine/Plugins/Runtime/ResonanceAudio/Source/ResonanceAudio/Private/ResonanceAudioDirectivityVisualizer.h b/Engine/Plugins/Runtime/ResonanceAudio/Source/ResonanceAudio/Private/ResonanceAudioDirectivityVisualizer.h
--- a/Engine/Plugins/Runtime/ResonanceAudio/Source/ResonanceAudio/Private/ResonanceAudioDirectivityVisualizer.h
+++ b/Engine/Plugins/Runtime/ResonanceAudio/Source/ResonanceAudio/Private/ResonanceAudioDirectivityVisualizer.h
@@ -28,6 +28,9 @@ public:
 	void SetSettings(UResonanceAudioSpatializationSourceSettings* InSettings) { Settings = InSettings; };
 	UResonanceAudioSpatializationSourceSettings* GetSettings() const { return Settings; };
 
+
+	virtual void BeginDestroy() override;
+
 private:
 	// Controls the smoothness of mesh visualization.
 	const int CIRCLE_SECTIONS = 128;
@@ -42,8 +45,9 @@ private:
 	TArray<FVector2D> UV0;
 	TArray<FProcMeshTangent> Tangents;
 
-	UPROPERTY()
+#if SUPPORTS_PROCEDURAL_MESH
 	UProceduralMeshComponent* Mesh;
+#endif // !PLATFORM_HTML5
 
 	UPROPERTY()
 	UMaterial* Material;
diff --git a/Engine/Plugins/Runtime/ResonanceAudio/Source/ResonanceAudio/ResonanceAudio.Build.cs b/Engine/Plugins/Runtime/ResonanceAudio/Source/ResonanceAudio/ResonanceAudio.Build.cs
--- a/Engine/Plugins/Runtime/ResonanceAudio/Source/ResonanceAudio/ResonanceAudio.Build.cs
+++ b/Engine/Plugins/Runtime/ResonanceAudio/Source/ResonanceAudio/ResonanceAudio.Build.cs
@@ -60,7 +60,6 @@ namespace UnrealBuildTool.Rules
             {
                 PrivateDependencyModuleNames.Add("UnrealEd");
                 PrivateDependencyModuleNames.Add("Landscape");
-                PrivateDependencyModuleNames.Add("ProceduralMeshComponent");
             }
 
             bEnableShadowVariableWarnings = false;
@@ -72,6 +71,17 @@ namespace UnrealBuildTool.Rules
                     "Eigen"
                     );
 
+			// Check whether Procedural Mesh is supported:
+			if(Target.Platform == UnrealTargetPlatform.HTML5)
+            {
+                PrivateDefinitions.Add("SUPPORTS_PROCEDURAL_MESH=0");
+            }
+            else
+            {
+                PrivateDependencyModuleNames.Add("ProceduralMeshComponent");
+                PrivateDefinitions.Add("SUPPORTS_PROCEDURAL_MESH=1");
+            }
+
             //Embree support:
             if (Target.Platform == UnrealTargetPlatform.Win64)
             {
