Fix Android build for Resonance plugin.
#rnx

diff --git a/Engine/Plugins/Runtime/ResonanceAudio/Source/ResonanceAudio/Private/ResonanceAudioLibrary/resonance_audio/base/simd_macros.h b/Engine/Plugins/Runtime/ResonanceAudio/Source/ResonanceAudio/Private/ResonanceAudioLibrary/resonance_audio/base/simd_macros.h
--- a/Engine/Plugins/Runtime/ResonanceAudio/Source/ResonanceAudio/Private/ResonanceAudioLibrary/resonance_audio/base/simd_macros.h
+++ b/Engine/Plugins/Runtime/ResonanceAudio/Source/ResonanceAudio/Private/ResonanceAudioLibrary/resonance_audio/base/simd_macros.h
@@ -58,7 +58,6 @@ typedef float SimdVector;
 #define SIMD_SQRT(a) (1.0f / FastReciprocalSqrt(a))
 #define SIMD_RECIPROCAL_SQRT(a) FastReciprocalSqrt(a)
 #define SIMD_LOAD_ONE_FLOAT(p) (p)
-#warning "Not using SIMD optimizations!"
 #endif
 
 #endif  // RESONANCE_AUDIO_BASE_SIMD_MACROS_H_
diff --git a/Engine/Plugins/Runtime/ResonanceAudio/Source/ResonanceAudio/ResonanceAudio.Build.cs b/Engine/Plugins/Runtime/ResonanceAudio/Source/ResonanceAudio/ResonanceAudio.Build.cs
--- a/Engine/Plugins/Runtime/ResonanceAudio/Source/ResonanceAudio/ResonanceAudio.Build.cs
+++ b/Engine/Plugins/Runtime/ResonanceAudio/Source/ResonanceAudio/ResonanceAudio.Build.cs
@@ -82,6 +82,12 @@ namespace UnrealBuildTool.Rules
                 PrivateDefinitions.Add("SUPPORTS_PROCEDURAL_MESH=1");
             }
 
+			if(Target.Platform == UnrealTargetPlatform.Android)
+            {
+                PrivateDefinitions.Add("PFFFT_SIMD_DISABLE=1");
+                PrivateDefinitions.Add("EIGEN_HAS_CXX11_MATH=0");
+            }
+
             //Embree support:
             if (Target.Platform == UnrealTargetPlatform.Win64)
             {
