works as expected.

diff --git a/Engine/Plugins/Runtime/ResonanceAudio/Source/ResonanceAudio/Private/ResonanceAudioLibrary/resonance_audio/base/misc_math.h b/Engine/Plugins/Runtime/ResonanceAudio/Source/ResonanceAudio/Private/ResonanceAudioLibrary/resonance_audio/base/misc_math.h
--- a/Engine/Plugins/Runtime/ResonanceAudio/Source/ResonanceAudio/Private/ResonanceAudioLibrary/resonance_audio/base/misc_math.h
+++ b/Engine/Plugins/Runtime/ResonanceAudio/Source/ResonanceAudio/Private/ResonanceAudioLibrary/resonance_audio/base/misc_math.h
@@ -29,9 +29,13 @@ limitations under the License.
 #include <vector>
 
 #include "base/integral_types.h"
+
+#include "CoreMinimal.h"
+
 THIRD_PARTY_INCLUDES_START
 #include "Eigen/Dense"
 THIRD_PARTY_INCLUDES_END
+
 #include "base/constants_and_types.h"
 #include "base/logging.h"
 
