[NO NON-PII COMMIT MESSAGE]

diff --git a/Engine/Plugins/Runtime/ResonanceAudio/Source/ResonanceAudio/Private/ResonanceAudioLibrary/resonance_audio/geometrical_acoustics/sphere.h b/Engine/Plugins/Runtime/ResonanceAudio/Source/ResonanceAudio/Private/ResonanceAudioLibrary/resonance_audio/geometrical_acoustics/sphere.h
--- a/Engine/Plugins/Runtime/ResonanceAudio/Source/ResonanceAudio/Private/ResonanceAudioLibrary/resonance_audio/geometrical_acoustics/sphere.h
+++ b/Engine/Plugins/Runtime/ResonanceAudio/Source/ResonanceAudio/Private/ResonanceAudioLibrary/resonance_audio/geometrical_acoustics/sphere.h
@@ -23,7 +23,7 @@ limitations under the License.
 namespace vraudio {
 
 // A user defined sphere geometry that enables Embree to test ray intersections.
-struct RTCORE_ALIGN(16) Sphere {
+struct alignas(16) Sphere {
   // Center of the sphere.
   float center[3];
 
@@ -32,6 +32,9 @@ struct RTCORE_ALIGN(16) Sphere {
 
   // Geometry Id. This will be reported by ray intersections.
   unsigned int geometry_id;
+
+  // This struct is explicitly padded out to 32 bytes.
+  uint8 padding[12];
 };
 
 // The following functions are to be called by Embree internally to enable fast
