Fix for compile in non-unity builds. use int for padding rather than uint8.

diff --git a/Engine/Plugins/Runtime/ResonanceAudio/Source/ResonanceAudio/Private/ResonanceAudioLibrary/resonance_audio/geometrical_acoustics/sphere.h b/Engine/Plugins/Runtime/ResonanceAudio/Source/ResonanceAudio/Private/ResonanceAudioLibrary/resonance_audio/geometrical_acoustics/sphere.h
--- a/Engine/Plugins/Runtime/ResonanceAudio/Source/ResonanceAudio/Private/ResonanceAudioLibrary/resonance_audio/geometrical_acoustics/sphere.h
+++ b/Engine/Plugins/Runtime/ResonanceAudio/Source/ResonanceAudio/Private/ResonanceAudioLibrary/resonance_audio/geometrical_acoustics/sphere.h
@@ -34,7 +34,7 @@ struct alignas(16) Sphere {
   unsigned int geometry_id;
 
   // This struct is explicitly padded out to 32 bytes.
-  uint8 padding[12];
+  int padding[3];
 };
 
 // The following functions are to be called by Embree internally to enable fast
