Add ResonanceEigenHelper.h, which suppresses static analysis warnings for Eigen.

diff --git a/Engine/Plugins/Runtime/ResonanceAudio/Source/ResonanceAudio/Private/ResonanceAudioLibrary/resonance_audio/ambisonics/ambisonic_codec_impl.h b/Engine/Plugins/Runtime/ResonanceAudio/Source/ResonanceAudio/Private/ResonanceAudioLibrary/resonance_audio/ambisonics/ambisonic_codec_impl.h
--- a/Engine/Plugins/Runtime/ResonanceAudio/Source/ResonanceAudio/Private/ResonanceAudioLibrary/resonance_audio/ambisonics/ambisonic_codec_impl.h
+++ b/Engine/Plugins/Runtime/ResonanceAudio/Source/ResonanceAudio/Private/ResonanceAudioLibrary/resonance_audio/ambisonics/ambisonic_codec_impl.h
@@ -20,7 +20,7 @@ limitations under the License.
 #include <cmath>
 #include <vector>
 
-#include "Eigen/Dense"
+#include "ResonanceEigenHelper.h"
 #include "ambisonics/ambisonic_codec.h"
 #include "ambisonics/associated_legendre_polynomials_generator.h"
 #include "ambisonics/utils.h"
diff --git a/Engine/Plugins/Runtime/ResonanceAudio/Source/ResonanceAudio/Private/ResonanceAudioLibrary/resonance_audio/ambisonics/hoa_rotator.h b/Engine/Plugins/Runtime/ResonanceAudio/Source/ResonanceAudio/Private/ResonanceAudioLibrary/resonance_audio/ambisonics/hoa_rotator.h
--- a/Engine/Plugins/Runtime/ResonanceAudio/Source/ResonanceAudio/Private/ResonanceAudioLibrary/resonance_audio/ambisonics/hoa_rotator.h
+++ b/Engine/Plugins/Runtime/ResonanceAudio/Source/ResonanceAudio/Private/ResonanceAudioLibrary/resonance_audio/ambisonics/hoa_rotator.h
@@ -20,7 +20,7 @@ limitations under the License.
 #include <vector>
 
 
-#include "Eigen/Dense"
+#include "ResonanceEigenHelper.h"
 #include "base/audio_buffer.h"
 #include "base/misc_math.h"
 
diff --git a/Engine/Plugins/Runtime/ResonanceAudio/Source/ResonanceAudio/Private/ResonanceAudioLibrary/resonance_audio/base/misc_math.h b/Engine/Plugins/Runtime/ResonanceAudio/Source/ResonanceAudio/Private/ResonanceAudioLibrary/resonance_audio/base/misc_math.h
--- a/Engine/Plugins/Runtime/ResonanceAudio/Source/ResonanceAudio/Private/ResonanceAudioLibrary/resonance_audio/base/misc_math.h
+++ b/Engine/Plugins/Runtime/ResonanceAudio/Source/ResonanceAudio/Private/ResonanceAudioLibrary/resonance_audio/base/misc_math.h
@@ -29,7 +29,7 @@ limitations under the License.
 #include <vector>
 
 #include "base/integral_types.h"
-#include "Eigen/Dense"
+#include "ResonanceEigenHelper.h"
 #include "base/constants_and_types.h"
 #include "base/logging.h"
 
diff --git a/Engine/Plugins/Runtime/ResonanceAudio/Source/ResonanceAudio/Private/ResonanceAudioLibrary/resonance_audio/dsp/biquad_filter.h b/Engine/Plugins/Runtime/ResonanceAudio/Source/ResonanceAudio/Private/ResonanceAudioLibrary/resonance_audio/dsp/biquad_filter.h
--- a/Engine/Plugins/Runtime/ResonanceAudio/Source/ResonanceAudio/Private/ResonanceAudioLibrary/resonance_audio/dsp/biquad_filter.h
+++ b/Engine/Plugins/Runtime/ResonanceAudio/Source/ResonanceAudio/Private/ResonanceAudioLibrary/resonance_audio/dsp/biquad_filter.h
@@ -22,7 +22,7 @@ limitations under the License.
 #include "base/audio_buffer.h"
 #include "base/logging.h"
 
-#include "Eigen/Core"
+#include "ResonanceEigenHelper.h"
 #include <array>
 
 namespace vraudio {
diff --git a/Engine/Plugins/Runtime/ResonanceAudio/Source/ResonanceAudio/Private/ResonanceAudioLibrary/resonance_audio/geometrical_acoustics/acoustic_listener.h b/Engine/Plugins/Runtime/ResonanceAudio/Source/ResonanceAudio/Private/ResonanceAudioLibrary/resonance_audio/geometrical_acoustics/acoustic_listener.h
--- a/Engine/Plugins/Runtime/ResonanceAudio/Source/ResonanceAudio/Private/ResonanceAudioLibrary/resonance_audio/geometrical_acoustics/acoustic_listener.h
+++ b/Engine/Plugins/Runtime/ResonanceAudio/Source/ResonanceAudio/Private/ResonanceAudioLibrary/resonance_audio/geometrical_acoustics/acoustic_listener.h
@@ -21,7 +21,7 @@ limitations under the License.
 #include <cstddef>
 #include <vector>
 
-#include "Eigen/Core"
+#include "ResonanceEigenHelper.h"
 #include "base/constants_and_types.h"
 
 namespace vraudio {
diff --git a/Engine/Plugins/Runtime/ResonanceAudio/Source/ResonanceAudio/Private/ResonanceAudioLibrary/resonance_audio/geometrical_acoustics/acoustic_source.h b/Engine/Plugins/Runtime/ResonanceAudio/Source/ResonanceAudio/Private/ResonanceAudioLibrary/resonance_audio/geometrical_acoustics/acoustic_source.h
--- a/Engine/Plugins/Runtime/ResonanceAudio/Source/ResonanceAudio/Private/ResonanceAudioLibrary/resonance_audio/geometrical_acoustics/acoustic_source.h
+++ b/Engine/Plugins/Runtime/ResonanceAudio/Source/ResonanceAudio/Private/ResonanceAudioLibrary/resonance_audio/geometrical_acoustics/acoustic_source.h
@@ -22,7 +22,9 @@ limitations under the License.
 #include <utility>
 #include <vector>
 
-#include "Eigen/Core"
+
+#include "ResonanceEigenHelper.h"
+
 #include "base/constants_and_types.h"
 #include "geometrical_acoustics/acoustic_ray.h"
 #include "geometrical_acoustics/sampling.h"
diff --git a/Engine/Plugins/Runtime/ResonanceAudio/Source/ResonanceAudio/Private/ResonanceAudioLibrary/resonance_audio/geometrical_acoustics/collection_kernel.cc b/Engine/Plugins/Runtime/ResonanceAudio/Source/ResonanceAudio/Private/ResonanceAudioLibrary/resonance_audio/geometrical_acoustics/collection_kernel.cc
--- a/Engine/Plugins/Runtime/ResonanceAudio/Source/ResonanceAudio/Private/ResonanceAudioLibrary/resonance_audio/geometrical_acoustics/collection_kernel.cc
+++ b/Engine/Plugins/Runtime/ResonanceAudio/Source/ResonanceAudio/Private/ResonanceAudioLibrary/resonance_audio/geometrical_acoustics/collection_kernel.cc
@@ -20,7 +20,7 @@ limitations under the License.
 #include <cmath>
 #include <vector>
 
-#include "Eigen/Core"
+#include "ResonanceEigenHelper.h"
 #include "base/constants_and_types.h"
 #include "base/logging.h"
 
diff --git a/Engine/Plugins/Runtime/ResonanceAudio/Source/ResonanceAudio/Private/ResonanceAudioLibrary/resonance_audio/geometrical_acoustics/impulse_response_computer.cc b/Engine/Plugins/Runtime/ResonanceAudio/Source/ResonanceAudio/Private/ResonanceAudioLibrary/resonance_audio/geometrical_acoustics/impulse_response_computer.cc
--- a/Engine/Plugins/Runtime/ResonanceAudio/Source/ResonanceAudio/Private/ResonanceAudioLibrary/resonance_audio/geometrical_acoustics/impulse_response_computer.cc
+++ b/Engine/Plugins/Runtime/ResonanceAudio/Source/ResonanceAudio/Private/ResonanceAudioLibrary/resonance_audio/geometrical_acoustics/impulse_response_computer.cc
@@ -18,7 +18,7 @@ limitations under the License.
 
 #include <utility>
 
-#include "Eigen/Core"
+#include "ResonanceEigenHelper.h"
 #include "base/logging.h"
 #include "geometrical_acoustics/parallel_for.h"
 #include "geometrical_acoustics/reflection_kernel.h"
diff --git a/Engine/Plugins/Runtime/ResonanceAudio/Source/ResonanceAudio/Private/ResonanceAudioLibrary/resonance_audio/geometrical_acoustics/proxy_room_estimator.cc b/Engine/Plugins/Runtime/ResonanceAudio/Source/ResonanceAudio/Private/ResonanceAudioLibrary/resonance_audio/geometrical_acoustics/proxy_room_estimator.cc
--- a/Engine/Plugins/Runtime/ResonanceAudio/Source/ResonanceAudio/Private/ResonanceAudioLibrary/resonance_audio/geometrical_acoustics/proxy_room_estimator.cc
+++ b/Engine/Plugins/Runtime/ResonanceAudio/Source/ResonanceAudio/Private/ResonanceAudioLibrary/resonance_audio/geometrical_acoustics/proxy_room_estimator.cc
@@ -19,7 +19,7 @@ limitations under the License.
 #include <algorithm>
 #include <cmath>
 
-#include "Eigen/Core"
+#include "ResonanceEigenHelper.h"
 #include "base/logging.h"
 #include "platforms/common/room_effects_utils.h"
 
diff --git a/Engine/Plugins/Runtime/ResonanceAudio/Source/ResonanceAudio/Private/ResonanceAudioLibrary/resonance_audio/geometrical_acoustics/proxy_room_estimator.h b/Engine/Plugins/Runtime/ResonanceAudio/Source/ResonanceAudio/Private/ResonanceAudioLibrary/resonance_audio/geometrical_acoustics/proxy_room_estimator.h
--- a/Engine/Plugins/Runtime/ResonanceAudio/Source/ResonanceAudio/Private/ResonanceAudioLibrary/resonance_audio/geometrical_acoustics/proxy_room_estimator.h
+++ b/Engine/Plugins/Runtime/ResonanceAudio/Source/ResonanceAudio/Private/ResonanceAudioLibrary/resonance_audio/geometrical_acoustics/proxy_room_estimator.h
@@ -21,7 +21,7 @@ limitations under the License.
 #include <utility>
 #include <vector>
 
-#include "Eigen/Core"
+#include "ResonanceEigenHelper.h"
 #include "api/resonance_audio_api.h"
 #include "base/constants_and_types.h"
 #include "base/misc_math.h"
diff --git a/Engine/Plugins/Runtime/ResonanceAudio/Source/ResonanceAudio/Private/ResonanceAudioLibrary/resonance_audio/geometrical_acoustics/reflection_kernel.h b/Engine/Plugins/Runtime/ResonanceAudio/Source/ResonanceAudio/Private/ResonanceAudioLibrary/resonance_audio/geometrical_acoustics/reflection_kernel.h
--- a/Engine/Plugins/Runtime/ResonanceAudio/Source/ResonanceAudio/Private/ResonanceAudioLibrary/resonance_audio/geometrical_acoustics/reflection_kernel.h
+++ b/Engine/Plugins/Runtime/ResonanceAudio/Source/ResonanceAudio/Private/ResonanceAudioLibrary/resonance_audio/geometrical_acoustics/reflection_kernel.h
@@ -17,7 +17,7 @@ limitations under the License.
 #ifndef RESONANCE_AUDIO_GEOMETRICAL_ACOUSTICS_REFLECTION_KERNEL_H_
 #define RESONANCE_AUDIO_GEOMETRICAL_ACOUSTICS_REFLECTION_KERNEL_H_
 
-#include "Eigen/Core"
+#include "ResonanceEigenHelper.h"
 #include "base/constants_and_types.h"
 #include "base/logging.h"
 #include "geometrical_acoustics/acoustic_ray.h"
diff --git a/Engine/Plugins/Runtime/ResonanceAudio/Source/ResonanceAudio/Private/ResonanceAudioLibrary/resonance_audio/geometrical_acoustics/sampling.h b/Engine/Plugins/Runtime/ResonanceAudio/Source/ResonanceAudio/Private/ResonanceAudioLibrary/resonance_audio/geometrical_acoustics/sampling.h
--- a/Engine/Plugins/Runtime/ResonanceAudio/Source/ResonanceAudio/Private/ResonanceAudioLibrary/resonance_audio/geometrical_acoustics/sampling.h
+++ b/Engine/Plugins/Runtime/ResonanceAudio/Source/ResonanceAudio/Private/ResonanceAudioLibrary/resonance_audio/geometrical_acoustics/sampling.h
@@ -21,7 +21,7 @@ limitations under the License.
 
 #include <cmath>
 
-#include "Eigen/Dense"
+#include "ResonanceEigenHelper.h"
 #include "base/constants_and_types.h"
 #include "base/logging.h"
 
diff --git a/Engine/Plugins/Runtime/ResonanceAudio/Source/ResonanceAudio/Private/ResonanceAudioLibrary/resonance_audio/geometrical_acoustics/sphere.cc b/Engine/Plugins/Runtime/ResonanceAudio/Source/ResonanceAudio/Private/ResonanceAudioLibrary/resonance_audio/geometrical_acoustics/sphere.cc
--- a/Engine/Plugins/Runtime/ResonanceAudio/Source/ResonanceAudio/Private/ResonanceAudioLibrary/resonance_audio/geometrical_acoustics/sphere.cc
+++ b/Engine/Plugins/Runtime/ResonanceAudio/Source/ResonanceAudio/Private/ResonanceAudioLibrary/resonance_audio/geometrical_acoustics/sphere.cc
@@ -19,7 +19,7 @@ limitations under the License.
 #include <cmath>
 #include <utility>
 
-#include "Eigen/Core"
+#include "ResonanceEigenHelper.h"
 
 namespace vraudio {
 
diff --git a/Engine/Plugins/Runtime/ResonanceAudio/Source/ResonanceAudio/Private/ResonanceAudioLibrary/resonance_audio/utils/pseudoinverse.h b/Engine/Plugins/Runtime/ResonanceAudio/Source/ResonanceAudio/Private/ResonanceAudioLibrary/resonance_audio/utils/pseudoinverse.h
--- a/Engine/Plugins/Runtime/ResonanceAudio/Source/ResonanceAudio/Private/ResonanceAudioLibrary/resonance_audio/utils/pseudoinverse.h
+++ b/Engine/Plugins/Runtime/ResonanceAudio/Source/ResonanceAudio/Private/ResonanceAudioLibrary/resonance_audio/utils/pseudoinverse.h
@@ -17,7 +17,7 @@ limitations under the License.
 #ifndef RESONANCE_AUDIO_UTILS_PSEUDOINVERSE_H_
 #define RESONANCE_AUDIO_UTILS_PSEUDOINVERSE_H_
 
-#include "Eigen/Dense"
+#include "ResonanceEigenHelper.h"
 
 namespace vraudio {
 
diff --git a/Engine/Plugins/Runtime/ResonanceAudio/Source/ResonanceAudio/Private/ResonanceEigenHelper.h b/Engine/Plugins/Runtime/ResonanceAudio/Source/ResonanceAudio/Private/ResonanceEigenHelper.h
new file mode 100644
--- /dev/null
+++ b/Engine/Plugins/Runtime/ResonanceAudio/Source/ResonanceAudio/Private/ResonanceEigenHelper.h
@@ -0,0 +1,17 @@
+// Copyright 1998-2019 Epic Games, Inc. All Rights Reserved.
+
+#pragma once
+
+// This file provides a way to include Eigen/Core without producing static analysis warnings.
+
+#if defined(_MSC_VER) && USING_CODE_ANALYSIS
+#pragma warning(push)
+#pragma warning(disable:6294) // Ill-defined for-loop:  initial condition does not satisfy test.  Loop body not executed.
+#endif
+THIRD_PARTY_INCLUDES_START
+#include "Eigen/Core"
+#include "Eigen/Dense"
+THIRD_PARTY_INCLUDES_END
+#if defined(_MSC_VER) && USING_CODE_ANALYSIS
+#pragma warning(pop)
+#endif
\ No newline at end of file
