Avoid using CoreMinimal.h or the THIRD_PARTY_INCLUDES macro in ResonanceEigenHelper, since it doesn't not gel with Eigen's macros.h.

diff --git a/Engine/Plugins/Runtime/ResonanceAudio/Source/ResonanceAudio/Private/ResonanceEigenHelper.h b/Engine/Plugins/Runtime/ResonanceAudio/Source/ResonanceAudio/Private/ResonanceEigenHelper.h
--- a/Engine/Plugins/Runtime/ResonanceAudio/Source/ResonanceAudio/Private/ResonanceEigenHelper.h
+++ b/Engine/Plugins/Runtime/ResonanceAudio/Source/ResonanceAudio/Private/ResonanceEigenHelper.h
@@ -3,15 +3,27 @@
 #pragma once
 
 // This file provides a way to include Eigen/Core without producing static analysis warnings.
-
-#if defined(_MSC_VER) && USING_CODE_ANALYSIS
+#if defined(__clang__)
+_Pragma("clang diagnostic push") \
+_Pragma("clang diagnostic ignored \"-Wshadow\"")
+#elif defined(_MSC_VER)
 #pragma warning(push)
-#pragma warning(disable:6294) // Ill-defined for-loop:  initial condition does not satisfy test.  Loop body not executed.
+#pragma warning(disable:6294) /* Ill-defined for-loop:  initial condition does not satisfy test.  Loop body not executed. */
+#pragma warning(disable:6326) /* Potential comparison of a constant with another constant. */
+#pragma warning(disable:4456) /* declaration of 'LocalVariable' hides previous local declaration */ 
+#pragma warning(disable:4457) /* declaration of 'LocalVariable' hides function parameter */ 
+#pragma warning(disable:4458) /* declaration of 'LocalVariable' hides class member */ 
+#pragma warning(disable:4459) /* declaration of 'LocalVariable' hides global declaration */ 
+#pragma warning(disable:6244) /* local declaration of <variable> hides previous declaration at <line> of <file> */
 #endif
-THIRD_PARTY_INCLUDES_START
+
+
 #include "Eigen/Core"
 #include "Eigen/Dense"
-THIRD_PARTY_INCLUDES_END
-#if defined(_MSC_VER) && USING_CODE_ANALYSIS
+
+#if defined(__clang__)
+_Pragma("clang diagnostic pop")
+#elif defined(_MSC_VER)
 #pragma warning(pop)
-#endif
\ No newline at end of file
+#endif
+
