[NO NON-PII COMMIT MESSAGE]

diff --git a/Engine/Plugins/Runtime/ResonanceAudio/Source/ResonanceAudio/Private/NullEmbreeImplementation.cpp b/Engine/Plugins/Runtime/ResonanceAudio/Source/ResonanceAudio/Private/NullEmbreeImplementation.cpp
--- a/Engine/Plugins/Runtime/ResonanceAudio/Source/ResonanceAudio/Private/NullEmbreeImplementation.cpp
+++ b/Engine/Plugins/Runtime/ResonanceAudio/Source/ResonanceAudio/Private/NullEmbreeImplementation.cpp
@@ -5,8 +5,7 @@
 
 #if !USE_EMBREE
 
-#include "embree2/rtcore.h"
-#include "embree2/rtcore_ray.h"
+#include "ResonanceEmbreeHelpers.h"
 #include "CoreMinimal.h"
 
 void rtcCommit(RTCScene scene)
diff --git a/Engine/Plugins/Runtime/ResonanceAudio/Source/ResonanceAudio/Private/ResonanceAudioLibrary/resonance_audio/geometrical_acoustics/acoustic_ray.h b/Engine/Plugins/Runtime/ResonanceAudio/Source/ResonanceAudio/Private/ResonanceAudioLibrary/resonance_audio/geometrical_acoustics/acoustic_ray.h
--- a/Engine/Plugins/Runtime/ResonanceAudio/Source/ResonanceAudio/Private/ResonanceAudioLibrary/resonance_audio/geometrical_acoustics/acoustic_ray.h
+++ b/Engine/Plugins/Runtime/ResonanceAudio/Source/ResonanceAudio/Private/ResonanceAudioLibrary/resonance_audio/geometrical_acoustics/acoustic_ray.h
@@ -20,8 +20,7 @@ limitations under the License.
 #include <array>
 #include <vector>
 
-#include "embree2/rtcore.h"
-#include "embree2/rtcore_ray.h"
+#include "ResonanceEmbreeHelper.h"
 
 #include "base/constants_and_types.h"
 
diff --git a/Engine/Plugins/Runtime/ResonanceAudio/Source/ResonanceAudio/Private/ResonanceAudioLibrary/resonance_audio/geometrical_acoustics/scene_manager.h b/Engine/Plugins/Runtime/ResonanceAudio/Source/ResonanceAudio/Private/ResonanceAudioLibrary/resonance_audio/geometrical_acoustics/scene_manager.h
--- a/Engine/Plugins/Runtime/ResonanceAudio/Source/ResonanceAudio/Private/ResonanceAudioLibrary/resonance_audio/geometrical_acoustics/scene_manager.h
+++ b/Engine/Plugins/Runtime/ResonanceAudio/Source/ResonanceAudio/Private/ResonanceAudioLibrary/resonance_audio/geometrical_acoustics/scene_manager.h
@@ -22,8 +22,7 @@ limitations under the License.
 #include <unordered_set>
 #include <vector>
 
-#include "embree2/rtcore.h"
-#include "embree2/rtcore_ray.h"
+#include "ResonanceEmbreeHelper.h"
 
 #include "base/logging.h"
 #include "geometrical_acoustics/acoustic_listener.h"
diff --git a/Engine/Plugins/Runtime/ResonanceAudio/Source/ResonanceAudio/Private/ResonanceAudioLibrary/resonance_audio/geometrical_acoustics/sphere.h b/Engine/Plugins/Runtime/ResonanceAudio/Source/ResonanceAudio/Private/ResonanceAudioLibrary/resonance_audio/geometrical_acoustics/sphere.h
--- a/Engine/Plugins/Runtime/ResonanceAudio/Source/ResonanceAudio/Private/ResonanceAudioLibrary/resonance_audio/geometrical_acoustics/sphere.h
+++ b/Engine/Plugins/Runtime/ResonanceAudio/Source/ResonanceAudio/Private/ResonanceAudioLibrary/resonance_audio/geometrical_acoustics/sphere.h
@@ -17,8 +17,7 @@ limitations under the License.
 #ifndef RESONANCE_AUDIO_GEOMETRICAL_ACOUSTICS_SPHERE_H_
 #define RESONANCE_AUDIO_GEOMETRICAL_ACOUSTICS_SPHERE_H_
 
-#include "embree2/rtcore.h"
-#include "embree2/rtcore_ray.h"
+#include "ResonanceEmbreeHelper.h"
 
 namespace vraudio {
 
diff --git a/Engine/Plugins/Runtime/ResonanceAudio/Source/ResonanceAudio/Private/ResonanceEmbreeHelper.h b/Engine/Plugins/Runtime/ResonanceAudio/Source/ResonanceAudio/Private/ResonanceEmbreeHelper.h
new file mode 100644
--- /dev/null
+++ b/Engine/Plugins/Runtime/ResonanceAudio/Source/ResonanceAudio/Private/ResonanceEmbreeHelper.h
@@ -0,0 +1,30 @@
+// Copyright 1998-2019 Epic Games, Inc. All Rights Reserved.
+
+#pragma once
+
+// This file provides a way to include rtcore.h and rtcore_ray.h from Embree without producing static analysis warnings.
+#if defined(__clang__)
+_Pragma("clang diagnostic push") \
+_Pragma("clang diagnostic ignored \"-Wshadow\"")
+#elif defined(_MSC_VER)
+#pragma warning(push)
+#pragma warning(disable:6294) /* Ill-defined for-loop:  initial condition does not satisfy test.  Loop body not executed. */
+#pragma warning(disable:6326) /* Potential comparison of a constant with another constant. */
+#pragma warning(disable:4456) /* declaration of 'LocalVariable' hides previous local declaration */ 
+#pragma warning(disable:4457) /* declaration of 'LocalVariable' hides function parameter */ 
+#pragma warning(disable:4458) /* declaration of 'LocalVariable' hides class member */ 
+#pragma warning(disable:4459) /* declaration of 'LocalVariable' hides global declaration */ 
+#pragma warning(disable:6244) /* local declaration of <variable> hides previous declaration at <line> of <file> */
+#pragma warning(disable:4324) /* structure was padded due to alignment specifier */
+#endif
+
+
+#include "embree2/rtcore.h"
+#include "embree2/rtcore_ray.h"
+
+#if defined(__clang__)
+_Pragma("clang diagnostic pop")
+#elif defined(_MSC_VER)
+#pragma warning(pop)
+#endif
+
