Probable fix for include not found
#rnx

diff --git a/Engine/Plugins/Runtime/ResonanceAudio/Source/ResonanceAudio/Private/NullEmbreeImplementation.cpp b/Engine/Plugins/Runtime/ResonanceAudio/Source/ResonanceAudio/Private/NullEmbreeImplementation.cpp
--- a/Engine/Plugins/Runtime/ResonanceAudio/Source/ResonanceAudio/Private/NullEmbreeImplementation.cpp
+++ b/Engine/Plugins/Runtime/ResonanceAudio/Source/ResonanceAudio/Private/NullEmbreeImplementation.cpp
@@ -5,7 +5,7 @@
 
 #if !USE_EMBREE
 
-#include "ResonanceEmbreeHelpers.h"
+#include "ResonanceEmbreeHelper.h"
 #include "CoreMinimal.h"
 
 void rtcCommit(RTCScene scene)
