[NO NON-PII COMMIT MESSAGE]

diff --git a/Engine/Plugins/Runtime/ResonanceAudio/Source/ResonanceAudio/Private/ResonanceAudioLibrary/resonance_audio/utils/task_thread_pool.cc b/Engine/Plugins/Runtime/ResonanceAudio/Source/ResonanceAudio/Private/ResonanceAudioLibrary/resonance_audio/utils/task_thread_pool.cc
--- a/Engine/Plugins/Runtime/ResonanceAudio/Source/ResonanceAudio/Private/ResonanceAudioLibrary/resonance_audio/utils/task_thread_pool.cc
+++ b/Engine/Plugins/Runtime/ResonanceAudio/Source/ResonanceAudio/Private/ResonanceAudioLibrary/resonance_audio/utils/task_thread_pool.cc
@@ -14,6 +14,26 @@ See the License for the specific language governing permissions and
 limitations under the License.
 */
 
+#if defined(__clang__)
+_Pragma("clang diagnostic push") \
+_Pragma("clang diagnostic ignored \"-Wshadow\"")
+#elif defined(_MSC_VER)
+#pragma warning(push)
+#pragma warning(disable:6294) /* Ill-defined for-loop:  initial condition does not satisfy test.  Loop body not executed. */
+#pragma warning(disable:6326) /* Potential comparison of a constant with another constant. */
+#pragma warning(disable:4456) /* declaration of 'LocalVariable' hides previous local declaration */ 
+#pragma warning(disable:4457) /* declaration of 'LocalVariable' hides function parameter */ 
+#pragma warning(disable:4458) /* declaration of 'LocalVariable' hides class member */ 
+#pragma warning(disable:4459) /* declaration of 'LocalVariable' hides global declaration */ 
+#pragma warning(disable:6244) /* local declaration of <variable> hides previous declaration at <line> of <file> */
+#pragma warning(disable:4324) /* structure was padded due to alignment specifier */
+#pragma warning(disable:6011) /* Dereferencing NULL pointer '<ptr>' */
+#pragma warning(disable:6385) /* Reading invalid data from '<array>': the readable size is '<array size>' bytes, but '<read size>' bytes may be read */
+#pragma warning(disable:6386) /* Buffer overrun while writing to '<array>': the writable size is '<array size>' bytes, but '<write size>' bytes might be written */
+#pragma warning(disable:4244) /* Conversion from <type> to <type>, possible loss of data. */
+#pragma warning(disable:4702) /* unreachable code */
+#endif
+
 #include "utils/task_thread_pool.h"
 
 #include <thread>
@@ -247,3 +267,9 @@ void TaskThreadPool::WorkerThread::TaskLoop() {
 }
 
 }  // namespace vraudio
+
+#if defined(__clang__)
+_Pragma("clang diagnostic pop")
+#elif defined(_MSC_VER)
+#pragma warning(pop)
+#endif
diff --git a/Engine/Plugins/Runtime/ResonanceAudio/Source/ResonanceAudio/Private/ResonanceAudioLibrary/resonance_audio/utils/task_thread_pool.h b/Engine/Plugins/Runtime/ResonanceAudio/Source/ResonanceAudio/Private/ResonanceAudioLibrary/resonance_audio/utils/task_thread_pool.h
--- a/Engine/Plugins/Runtime/ResonanceAudio/Source/ResonanceAudio/Private/ResonanceAudioLibrary/resonance_audio/utils/task_thread_pool.h
+++ b/Engine/Plugins/Runtime/ResonanceAudio/Source/ResonanceAudio/Private/ResonanceAudioLibrary/resonance_audio/utils/task_thread_pool.h
@@ -17,6 +17,26 @@ limitations under the License.
 #ifndef RESONANCE_AUDIO_UTILS_TASK_THREAD_POOL_H_
 #define RESONANCE_AUDIO_UTILS_TASK_THREAD_POOL_H_
 
+#if defined(__clang__)
+_Pragma("clang diagnostic push") \
+_Pragma("clang diagnostic ignored \"-Wshadow\"")
+#elif defined(_MSC_VER)
+#pragma warning(push)
+#pragma warning(disable:6294) /* Ill-defined for-loop:  initial condition does not satisfy test.  Loop body not executed. */
+#pragma warning(disable:6326) /* Potential comparison of a constant with another constant. */
+#pragma warning(disable:4456) /* declaration of 'LocalVariable' hides previous local declaration */ 
+#pragma warning(disable:4457) /* declaration of 'LocalVariable' hides function parameter */ 
+#pragma warning(disable:4458) /* declaration of 'LocalVariable' hides class member */ 
+#pragma warning(disable:4459) /* declaration of 'LocalVariable' hides global declaration */ 
+#pragma warning(disable:6244) /* local declaration of <variable> hides previous declaration at <line> of <file> */
+#pragma warning(disable:4324) /* structure was padded due to alignment specifier */
+#pragma warning(disable:6011) /* Dereferencing NULL pointer '<ptr>' */
+#pragma warning(disable:6385) /* Reading invalid data from '<array>': the readable size is '<array size>' bytes, but '<read size>' bytes may be read */
+#pragma warning(disable:6386) /* Buffer overrun while writing to '<array>': the writable size is '<array size>' bytes, but '<write size>' bytes might be written */
+#pragma warning(disable:4244) /* Conversion from <type> to <type>, possible loss of data. */
+#pragma warning(disable:4702) /* unreachable code */
+#endif
+
 #include <atomic>
 #include <chrono>
 #include <condition_variable>
@@ -118,4 +138,10 @@ class TaskThreadPool {
 
 }  // namespace vraudio
 
+#if defined(__clang__)
+_Pragma("clang diagnostic pop")
+#elif defined(_MSC_VER)
+#pragma warning(pop)
+#endif
+
 #endif  // RESONANCE_AUDIO_UTILS_TASK_THREAD_POOL_H_
