Supress Unreferenced Formal Parameters in resonance audio plugin.

diff --git a/Engine/Plugins/Runtime/ResonanceAudio/Source/ResonanceAudio/Private/ResonanceAudioLibrary/resonance_audio/utils/task_thread_pool.cc b/Engine/Plugins/Runtime/ResonanceAudio/Source/ResonanceAudio/Private/ResonanceAudioLibrary/resonance_audio/utils/task_thread_pool.cc
--- a/Engine/Plugins/Runtime/ResonanceAudio/Source/ResonanceAudio/Private/ResonanceAudioLibrary/resonance_audio/utils/task_thread_pool.cc
+++ b/Engine/Plugins/Runtime/ResonanceAudio/Source/ResonanceAudio/Private/ResonanceAudioLibrary/resonance_audio/utils/task_thread_pool.cc
@@ -32,6 +32,7 @@ _Pragma("clang diagnostic ignored \"-Wshadow\"")
 #pragma warning(disable:6386) /* Buffer overrun while writing to '<array>': the writable size is '<array size>' bytes, but '<write size>' bytes might be written */
 #pragma warning(disable:4244) /* Conversion from <type> to <type>, possible loss of data. */
 #pragma warning(disable:4702) /* unreachable code */
+#pragma warning(disable:4100) /* unreferenced formal parameter */
 #endif
 
 #include "utils/task_thread_pool.h"
