[NO NON-PII COMMIT MESSAGE]

diff --git a/Engine/Plugins/Runtime/ResonanceAudio/Source/ResonanceAudio/Private/ResonanceAudioLibrary/resonance_audio/utils/lockless_task_queue.cc b/Engine/Plugins/Runtime/ResonanceAudio/Source/ResonanceAudio/Private/ResonanceAudioLibrary/resonance_audio/utils/lockless_task_queue.cc
--- a/Engine/Plugins/Runtime/ResonanceAudio/Source/ResonanceAudio/Private/ResonanceAudioLibrary/resonance_audio/utils/lockless_task_queue.cc
+++ b/Engine/Plugins/Runtime/ResonanceAudio/Source/ResonanceAudio/Private/ResonanceAudioLibrary/resonance_audio/utils/lockless_task_queue.cc
@@ -14,6 +14,22 @@ See the License for the specific language governing permissions and
 limitations under the License.
 */
 
+#if defined(__clang__)
+_Pragma("clang diagnostic push") \
+_Pragma("clang diagnostic ignored \"-Wshadow\"")
+#elif defined(_MSC_VER)
+#pragma warning(push)
+#pragma warning(disable:6294) /* Ill-defined for-loop:  initial condition does not satisfy test.  Loop body not executed. */
+#pragma warning(disable:6326) /* Potential comparison of a constant with another constant. */
+#pragma warning(disable:4456) /* declaration of 'LocalVariable' hides previous local declaration */ 
+#pragma warning(disable:4457) /* declaration of 'LocalVariable' hides function parameter */ 
+#pragma warning(disable:4458) /* declaration of 'LocalVariable' hides class member */ 
+#pragma warning(disable:4459) /* declaration of 'LocalVariable' hides global declaration */ 
+#pragma warning(disable:6244) /* local declaration of <variable> hides previous declaration at <line> of <file> */
+#pragma warning(disable:4702) /* unreachable code */
+#pragma warning(disable:4100) /* unreachable code */
+#endif
+
 #include "utils/lockless_task_queue.h"
 
 #include "base/logging.h"
@@ -116,3 +132,9 @@ void LocklessTaskQueue::Init(size_t num_nodes) {
 }
 
 }  // namespace vraudio
+
+#if defined(__clang__)
+_Pragma("clang diagnostic pop")
+#elif defined(_MSC_VER)
+#pragma warning(pop)
+#endif
\ No newline at end of file
diff --git a/Engine/Plugins/Runtime/ResonanceAudio/Source/ResonanceAudio/Private/ResonanceAudioLibrary/resonance_audio/utils/lockless_task_queue.h b/Engine/Plugins/Runtime/ResonanceAudio/Source/ResonanceAudio/Private/ResonanceAudioLibrary/resonance_audio/utils/lockless_task_queue.h
--- a/Engine/Plugins/Runtime/ResonanceAudio/Source/ResonanceAudio/Private/ResonanceAudioLibrary/resonance_audio/utils/lockless_task_queue.h
+++ b/Engine/Plugins/Runtime/ResonanceAudio/Source/ResonanceAudio/Private/ResonanceAudioLibrary/resonance_audio/utils/lockless_task_queue.h
@@ -17,6 +17,22 @@ limitations under the License.
 #ifndef RESONANCE_AUDIO_UTILS_LOCKLESS_TASK_QUEUE_H_
 #define RESONANCE_AUDIO_UTILS_LOCKLESS_TASK_QUEUE_H_
 
+#if defined(__clang__)
+_Pragma("clang diagnostic push") \
+_Pragma("clang diagnostic ignored \"-Wshadow\"")
+#elif defined(_MSC_VER)
+#pragma warning(push)
+#pragma warning(disable:6294) /* Ill-defined for-loop:  initial condition does not satisfy test.  Loop body not executed. */
+#pragma warning(disable:6326) /* Potential comparison of a constant with another constant. */
+#pragma warning(disable:4456) /* declaration of 'LocalVariable' hides previous local declaration */ 
+#pragma warning(disable:4457) /* declaration of 'LocalVariable' hides function parameter */ 
+#pragma warning(disable:4458) /* declaration of 'LocalVariable' hides class member */ 
+#pragma warning(disable:4459) /* declaration of 'LocalVariable' hides global declaration */ 
+#pragma warning(disable:6244) /* local declaration of <variable> hides previous declaration at <line> of <file> */
+#pragma warning(disable:4702) /* unreachable code */
+#pragma warning(disable:4100) /* unreachable code */
+#endif
+
 #include <stddef.h>
 
 #include <atomic>
@@ -56,7 +72,7 @@ class LocklessTaskQueue {
     Node() = default;
 
     // Dummy copy constructor to enable vector::resize allocation.
-    Node(const Node& node) : next() {}
+    Node(const Node& /*node*/) : next() {}
 
     // User task.
     LocklessTaskQueue::Task task;
@@ -105,4 +121,10 @@ class LocklessTaskQueue {
 
 }  // namespace vraudio
 
+#if defined(__clang__)
+_Pragma("clang diagnostic pop")
+#elif defined(_MSC_VER)
+#pragma warning(pop)
+#endif
+
 #endif  // RESONANCE_AUDIO_UTILS_LOCKLESS_TASK_QUEUE_H_
diff --git a/Engine/Plugins/Runtime/ResonanceAudio/Source/ResonanceAudio/Private/ResonanceAudioLibrary/resonance_audio/utils/vorbis_stream_encoder.h b/Engine/Plugins/Runtime/ResonanceAudio/Source/ResonanceAudio/Private/ResonanceAudioLibrary/resonance_audio/utils/vorbis_stream_encoder.h
--- a/Engine/Plugins/Runtime/ResonanceAudio/Source/ResonanceAudio/Private/ResonanceAudioLibrary/resonance_audio/utils/vorbis_stream_encoder.h
+++ b/Engine/Plugins/Runtime/ResonanceAudio/Source/ResonanceAudio/Private/ResonanceAudioLibrary/resonance_audio/utils/vorbis_stream_encoder.h
@@ -17,6 +17,26 @@ limitations under the License.
 #ifndef RESONANCE_AUDIO_UTILS_VORBIS_STREAM_ENCODER_H_
 #define RESONANCE_AUDIO_UTILS_VORBIS_STREAM_ENCODER_H_
 
+#if defined(__clang__)
+_Pragma("clang diagnostic push") \
+_Pragma("clang diagnostic ignored \"-Wshadow\"")
+#elif defined(_MSC_VER)
+#pragma warning(push)
+#pragma warning(disable:6294) /* Ill-defined for-loop:  initial condition does not satisfy test.  Loop body not executed. */
+#pragma warning(disable:6326) /* Potential comparison of a constant with another constant. */
+#pragma warning(disable:4456) /* declaration of 'LocalVariable' hides previous local declaration */ 
+#pragma warning(disable:4457) /* declaration of 'LocalVariable' hides function parameter */ 
+#pragma warning(disable:4458) /* declaration of 'LocalVariable' hides class member */ 
+#pragma warning(disable:4459) /* declaration of 'LocalVariable' hides global declaration */ 
+#pragma warning(disable:6244) /* local declaration of <variable> hides previous declaration at <line> of <file> */
+#pragma warning(disable:4324) /* structure was padded due to alignment specifier */
+#pragma warning(disable:6011) /* Dereferencing NULL pointer '<ptr>' */
+#pragma warning(disable:6385) /* Reading invalid data from '<array>': the readable size is '<array size>' bytes, but '<read size>' bytes may be read */
+#pragma warning(disable:6386) /* Buffer overrun while writing to '<array>': the writable size is '<array size>' bytes, but '<write size>' bytes might be written */
+#pragma warning(disable:4244) /* 'initializing': conversion from '<type>' to '<type>', possible loss of data */
+#pragma warning(disable:4121) /* '<type>' alignment of a member was sensitive to packing */
+#endif
+
 #include <fstream>
 #include <iostream>
 
@@ -111,4 +131,10 @@ class VorbisStreamEncoder {
 
 }  // namespace vraudio
 
+#if defined(__clang__)
+_Pragma("clang diagnostic pop")
+#elif defined(_MSC_VER)
+#pragma warning(pop)
+#endif
+
 #endif  // RESONANCE_AUDIO_UTILS_VORBIS_STREAM_ENCODER_H_
diff --git a/Engine/Plugins/Runtime/ResonanceAudio/Source/ResonanceAudio/Private/ResonanceEigenHelper.h b/Engine/Plugins/Runtime/ResonanceAudio/Source/ResonanceAudio/Private/ResonanceEigenHelper.h
--- a/Engine/Plugins/Runtime/ResonanceAudio/Source/ResonanceAudio/Private/ResonanceEigenHelper.h
+++ b/Engine/Plugins/Runtime/ResonanceAudio/Source/ResonanceAudio/Private/ResonanceEigenHelper.h
@@ -15,6 +15,7 @@ _Pragma("clang diagnostic ignored \"-Wshadow\"")
 #pragma warning(disable:4458) /* declaration of 'LocalVariable' hides class member */ 
 #pragma warning(disable:4459) /* declaration of 'LocalVariable' hides global declaration */ 
 #pragma warning(disable:6244) /* local declaration of <variable> hides previous declaration at <line> of <file> */
+#pragma warning(disable:4702) /* unreachable code */
 #endif
 
 
