Fix static analysis issue

diff --git a/Engine/Plugins/Runtime/ResonanceAudio/Source/ResonanceAudio/Private/ResonanceAudioLibrary/resonance_audio/utils/task_thread_pool.cc b/Engine/Plugins/Runtime/ResonanceAudio/Source/ResonanceAudio/Private/ResonanceAudioLibrary/resonance_audio/utils/task_thread_pool.cc
--- a/Engine/Plugins/Runtime/ResonanceAudio/Source/ResonanceAudio/Private/ResonanceAudioLibrary/resonance_audio/utils/task_thread_pool.cc
+++ b/Engine/Plugins/Runtime/ResonanceAudio/Source/ResonanceAudio/Private/ResonanceAudioLibrary/resonance_audio/utils/task_thread_pool.cc
@@ -58,7 +58,7 @@ class TaskThreadPool::WorkerThread {
   // std::vector<WorkerThread> container class.
   //
   // @param other |WorkerThread| instance this instance will be copied from.
-  WorkerThread(const TaskThreadPool::WorkerThread& other)
+  WorkerThread(const TaskThreadPool::WorkerThread& /*other*/)
       : parent_pool_(), task_closure_(), task_loop_triggered_(false) {}
 
   // Destructor.
