Skip linking Embree for mac in the Google Resonance plugin for now.

diff --git a/Engine/Plugins/Runtime/ResonanceAudio/Source/ResonanceAudio/ResonanceAudio.Build.cs b/Engine/Plugins/Runtime/ResonanceAudio/Source/ResonanceAudio/ResonanceAudio.Build.cs
--- a/Engine/Plugins/Runtime/ResonanceAudio/Source/ResonanceAudio/ResonanceAudio.Build.cs
+++ b/Engine/Plugins/Runtime/ResonanceAudio/Source/ResonanceAudio/ResonanceAudio.Build.cs
@@ -101,6 +101,11 @@ namespace UnrealBuildTool.Rules
                 RuntimeDependencies.Add("$(TargetOutputDir)/tbbmalloc.dll", SDKDir + "lib/tbbmalloc.dll");
                 PublicDefinitions.Add("USE_EMBREE=1");
             }
+			// Since using Embree for reflections is not supported in the current version of the Resonance plugin,
+			// we don't need to use it for now.
+			// To link Embree for Mac, uncomment the code below:
+
+			/*
             else if (Target.Platform == UnrealTargetPlatform.Mac)
             {
                 string SDKDir = Target.UEThirdPartySourceDirectory + "IntelEmbree/Embree2140/MacOSX/";
@@ -114,6 +119,7 @@ namespace UnrealBuildTool.Rules
                 RuntimeDependencies.Add("$(TargetOutputDir)/libtbbmalloc.dylib", SDKDir + "lib/libtbbmalloc.dylib");
                 PublicDefinitions.Add("USE_EMBREE=1");
             }
+			*/
             else
             {
 				// In platforms that don't support Embree, we implement no-op versions of the functions.
