Fix for macro redefinition in Resonance Audio plugin.

diff --git a/Engine/Plugins/Runtime/ResonanceAudio/Source/ResonanceAudio/ResonanceAudio.Build.cs b/Engine/Plugins/Runtime/ResonanceAudio/Source/ResonanceAudio/ResonanceAudio.Build.cs
--- a/Engine/Plugins/Runtime/ResonanceAudio/Source/ResonanceAudio/ResonanceAudio.Build.cs
+++ b/Engine/Plugins/Runtime/ResonanceAudio/Source/ResonanceAudio/ResonanceAudio.Build.cs
@@ -89,44 +89,13 @@ namespace UnrealBuildTool.Rules
             }
 
             //Embree support:
-            if (Target.Platform == UnrealTargetPlatform.Win64)
-            {
-                string SDKDir = Target.UEThirdPartySourceDirectory + "IntelEmbree/Embree2140/Win64/";
-
-                PublicIncludePaths.Add(SDKDir + "include");
-                PublicLibraryPaths.Add(SDKDir + "lib");
-                PublicAdditionalLibraries.Add("embree.2.14.0.lib");
-                RuntimeDependencies.Add("$(TargetOutputDir)/embree.2.14.0.dll", SDKDir + "lib/embree.2.14.0.dll");
-                RuntimeDependencies.Add("$(TargetOutputDir)/tbb.dll", SDKDir + "lib/tbb.dll");
-                RuntimeDependencies.Add("$(TargetOutputDir)/tbbmalloc.dll", SDKDir + "lib/tbbmalloc.dll");
-                PublicDefinitions.Add("USE_EMBREE=1");
-            }
-			// Since using Embree for reflections is not supported in the current version of the Resonance plugin,
-			// we don't need to use it for now.
-			// To link Embree for Mac, uncomment the code below:
+            string SDKDir = Target.UEThirdPartySourceDirectory + "IntelEmbree/Embree2140/Win64/";
+            PublicIncludePaths.Add(SDKDir + "include");
 
-			/*
-            else if (Target.Platform == UnrealTargetPlatform.Mac)
-            {
-                string SDKDir = Target.UEThirdPartySourceDirectory + "IntelEmbree/Embree2140/MacOSX/";
-
-                PublicIncludePaths.Add(SDKDir + "include");
-                PublicAdditionalLibraries.Add(SDKDir + "lib/libembree.2.14.0.dylib");
-                PublicAdditionalLibraries.Add(SDKDir + "lib/libtbb.dylib");
-                PublicAdditionalLibraries.Add(SDKDir + "lib/libtbbmalloc.dylib");
-                RuntimeDependencies.Add("$(TargetOutputDir)/libembree.2.14.0.dylib", SDKDir + "lib/libembree.2.14.0.dylib");
-                RuntimeDependencies.Add("$(TargetOutputDir)/libtbb.dylib", SDKDir + "lib/libtbb.dylib");
-                RuntimeDependencies.Add("$(TargetOutputDir)/libtbbmalloc.dylib", SDKDir + "lib/libtbbmalloc.dylib");
-                PublicDefinitions.Add("USE_EMBREE=1");
-            }
-			*/
-            else
+            if (!PublicDefinitions.Contains("USE_EMBREE=1"))
             {
-				// In platforms that don't support Embree, we implement no-op versions of the functions.
-                string SDKDir = Target.UEThirdPartySourceDirectory + "IntelEmbree/Embree2140/Win64/";
-                PublicIncludePaths.Add(SDKDir + "include");
-                PublicDefinitions.Add("USE_EMBREE=0");
-				PublicDefinitions.Add("EMBREE_STATIC_LIB=1");
+                // In platforms that don't support Embree, we implement no-op versions of the functions, which means we need to statically link embree.
+				PrivateDefinitions.Add("EMBREE_STATIC_LIB=1");
             }
         }
     }
