Couple spelling fixes.

diff --git a/Engine/Plugins/Runtime/ResonanceAudio/Source/ResonanceAudio/Public/ResonanceAudioSpatializationSourceSettings.h b/Engine/Plugins/Runtime/ResonanceAudio/Source/ResonanceAudio/Public/ResonanceAudioSpatializationSourceSettings.h
--- a/Engine/Plugins/Runtime/ResonanceAudio/Source/ResonanceAudio/Public/ResonanceAudioSpatializationSourceSettings.h
+++ b/Engine/Plugins/Runtime/ResonanceAudio/Source/ResonanceAudio/Public/ResonanceAudioSpatializationSourceSettings.h
@@ -58,7 +58,7 @@ public:
 	UPROPERTY(GlobalConfig, EditAnywhere, Category = "SpatializationSettings|Distance attenuation", meta = (ClampMin = "0.0", ClampMax = "1000000.0", UIMin = "0.0", UIMax = "1000000.0"))
 	float MinDistance;
 
-	// Maxium distance to apply the chosen attenuation method ((default = 50000.0 cm)
+	// Maximum distance to apply the chosen attenuation method ((default = 50000.0 cm)
 	UPROPERTY(GlobalConfig, EditAnywhere, Category = "SpatializationSettings|Distance attenuation", meta = (ClampMin = "0.0", ClampMax = "1000000.0", UIMin = "0.0", UIMax = "1000000.0"))
 	float MaxDistance;
 
diff --git a/Engine/Source/Runtime/Slate/Private/Framework/Application/SlateApplication.cpp b/Engine/Source/Runtime/Slate/Private/Framework/Application/SlateApplication.cpp
--- a/Engine/Source/Runtime/Slate/Private/Framework/Application/SlateApplication.cpp
+++ b/Engine/Source/Runtime/Slate/Private/Framework/Application/SlateApplication.cpp
@@ -6609,7 +6609,7 @@ bool FSlateApplication::OnTouchStarted( const TSharedPtr< FGenericWindow >& Plat
 	{
 #if !(UE_BUILD_SHIPPING || UE_BUILD_TEST)
 		// Only log when the touch starts, we don't want to spam the logs.
-		UE_LOG(LogSlate, Warning, TEXT("Maxium Touch Index Exceeded, %d, the maxium index allowed is %d"), TouchIndex, (((int32)ETouchIndex::CursorPointerIndex) - 1));
+		UE_LOG(LogSlate, Warning, TEXT("Maximum Touch Index Exceeded, %d, the maximum index allowed is %d"), TouchIndex, (((int32)ETouchIndex::CursorPointerIndex) - 1));
 #endif
 		return false;
 	}
