Allowing Google resonance to make a global spatialization source setting if one is not specified in the global plugin preset.

diff --git a/Engine/Plugins/Runtime/ResonanceAudio/Source/ResonanceAudio/Private/ResonanceAudioModule.cpp b/Engine/Plugins/Runtime/ResonanceAudio/Source/ResonanceAudio/Private/ResonanceAudioModule.cpp
--- a/Engine/Plugins/Runtime/ResonanceAudio/Source/ResonanceAudio/Private/ResonanceAudioModule.cpp
+++ b/Engine/Plugins/Runtime/ResonanceAudio/Source/ResonanceAudio/Private/ResonanceAudioModule.cpp
@@ -50,7 +50,8 @@ namespace ResonanceAudio
 			{
 				GlobalSpatializationSourceSettings = LoadObject<UResonanceAudioSpatializationSourceSettings>(nullptr, *GlobalPluginPresetName.ToString());
 			}
-			else
+			
+			if (!GlobalSpatializationSourceSettings)
 			{
 				GlobalSpatializationSourceSettings = NewObject<UResonanceAudioSpatializationSourceSettings>(UResonanceAudioSpatializationSourceSettings::StaticClass(), TEXT("Default Global Resonance Spatialization Preset"));
 			}
