Fix for Embree lib public definition causing packaging issues on unsupported platforms in ResonanceAudio.Build.cs

diff --git a/Engine/Plugins/Runtime/ResonanceAudio/Source/ResonanceAudio/ResonanceAudio.Build.cs b/Engine/Plugins/Runtime/ResonanceAudio/Source/ResonanceAudio/ResonanceAudio.Build.cs
--- a/Engine/Plugins/Runtime/ResonanceAudio/Source/ResonanceAudio/ResonanceAudio.Build.cs
+++ b/Engine/Plugins/Runtime/ResonanceAudio/Source/ResonanceAudio/ResonanceAudio.Build.cs
@@ -106,13 +106,8 @@ namespace UnrealBuildTool.Rules
                 string SDKDir = Target.UEThirdPartySourceDirectory + "IntelEmbree/Embree2140/MacOSX/";
 
                 PublicIncludePaths.Add(SDKDir + "include");
-                PublicAdditionalLibraries.Add(SDKDir + "lib/libembree.2.14.0.dylib");
-                PublicAdditionalLibraries.Add(SDKDir + "lib/libtbb.dylib");
-                PublicAdditionalLibraries.Add(SDKDir + "lib/libtbbmalloc.dylib");
-                RuntimeDependencies.Add("$(TargetOutputDir)/libembree.2.14.0.dylib", SDKDir + "lib/libembree.2.14.0.dylib");
-                RuntimeDependencies.Add("$(TargetOutputDir)/libtbb.dylib", SDKDir + "lib/libtbb.dylib");
-                RuntimeDependencies.Add("$(TargetOutputDir)/libtbbmalloc.dylib", SDKDir + "lib/libtbbmalloc.dylib");
-                PublicDefinitions.Add("USE_EMBREE=1");
+
+                PrivateDefinitions.Add("USE_EMBREE=1");
             }
             else
             {
@@ -120,7 +115,7 @@ namespace UnrealBuildTool.Rules
                 string SDKDir = Target.UEThirdPartySourceDirectory + "IntelEmbree/Embree2140/Win64/";
                 PublicIncludePaths.Add(SDKDir + "include");
                 PublicDefinitions.Add("USE_EMBREE=0");
-				PublicDefinitions.Add("EMBREE_STATIC_LIB=1");
+				PrivateDefinitions.Add("EMBREE_STATIC_LIB=1");
             }
         }
     }
