[NO NON-PII COMMIT MESSAGE]

diff --git a/Engine/Plugins/Runtime/ResonanceAudio/Source/ResonanceAudio/Private/ResonanceAudioModule.cpp b/Engine/Plugins/Runtime/ResonanceAudio/Source/ResonanceAudio/Private/ResonanceAudioModule.cpp
--- a/Engine/Plugins/Runtime/ResonanceAudio/Source/ResonanceAudio/Private/ResonanceAudioModule.cpp
+++ b/Engine/Plugins/Runtime/ResonanceAudio/Source/ResonanceAudio/Private/ResonanceAudioModule.cpp
@@ -37,11 +37,6 @@ namespace ResonanceAudio
 		IModularFeatures::Get().RegisterModularFeature(FSpatializationPluginFactory::GetModularFeatureName(), &SpatializationPluginFactory);
 		IModularFeatures::Get().RegisterModularFeature(FReverbPluginFactory::GetModularFeatureName(), &ReverbPluginFactory);
 
-		if (!ResonanceAudioDynamicLibraryHandle)
-		{
-			ResonanceAudioDynamicLibraryHandle = LoadResonanceAudioDynamicLibrary();
-		}
-
 		if (!IsRunningDedicatedServer() && !GlobalSpatializationSourceSettings)
 		{
 			// Load the global source preset settings:
