Resonance geometry fix

diff --git a/Engine/Plugins/Runtime/ResonanceAudio/Source/ResonanceAudio/Private/ResonanceAudioPluginListener.cpp b/Engine/Plugins/Runtime/ResonanceAudio/Source/ResonanceAudio/Private/ResonanceAudioPluginListener.cpp
--- a/Engine/Plugins/Runtime/ResonanceAudio/Source/ResonanceAudio/Private/ResonanceAudioPluginListener.cpp
+++ b/Engine/Plugins/Runtime/ResonanceAudio/Source/ResonanceAudio/Private/ResonanceAudioPluginListener.cpp
@@ -6,6 +6,7 @@
 #include "ResonanceAudioAmbisonics.h"
 #include "AudioDevice.h"
 #include "Async/Async.h"
+#include "Components/BrushComponent.h"
 
 namespace ResonanceAudio
 {
@@ -108,8 +109,10 @@ namespace ResonanceAudio
 					const FQuat CurrentVolumeRotation = CurrentVolume->GetActorQuat();
 					Preset->SetRoomRotation(CurrentVolumeRotation);
 					const FVector CurrentVolumeDimensions = CurrentVolume->GetActorScale3D();
+					const FVector CurrentBrushShapeExtents = 2.0f * CurrentVolume->GetBrushComponent()->Brush->Bounds.BoxExtent;
+					const FVector RoomDimensions = CurrentVolumeDimensions * CurrentBrushShapeExtents;
 					// The default Audio Volume cube size is 200cm, please see UCubeBuilder constructor for initialization details.
-					Preset->SetRoomDimensions(CurrentVolumeDimensions * 200.f);
+					Preset->SetRoomDimensions(RoomDimensions);
 				}
 				// Activate this preset or no room effects if nullptr.
 				ReverbPtr->SetPreset(Preset);
