Reduce log level of spamming google resonance logs

diff --git a/Engine/Plugins/Runtime/ResonanceAudio/Source/ResonanceAudio/Private/ResonanceAudioSpatialization.cpp b/Engine/Plugins/Runtime/ResonanceAudio/Source/ResonanceAudio/Private/ResonanceAudioSpatialization.cpp
--- a/Engine/Plugins/Runtime/ResonanceAudio/Source/ResonanceAudio/Private/ResonanceAudioSpatialization.cpp
+++ b/Engine/Plugins/Runtime/ResonanceAudio/Source/ResonanceAudio/Private/ResonanceAudioSpatialization.cpp
@@ -75,26 +75,26 @@ namespace ResonanceAudio {
 		if (SpatializationSettings[SourceId]->SpatializationMethod == ERaSpatializationMethod::STEREO_PANNING || GetDefault<UResonanceAudioSettings>()->QualityMode == ERaQualityMode::STEREO_PANNING)
 		{
 			BinauralSource.Id = ResonanceAudioApi->CreateSoundObjectSource(vraudio::RenderingMode::kStereoPanning);
-			UE_LOG(LogResonanceAudio, Log, TEXT("ResonanceAudioSpatializer::OnInitSource: STEREO_PANNING mode chosen."));
+			UE_LOG(LogResonanceAudio, Verbose, TEXT("ResonanceAudioSpatializer::OnInitSource: STEREO_PANNING mode chosen."));
 		}
 		else if (GetDefault<UResonanceAudioSettings>()->QualityMode == ERaQualityMode::BINAURAL_LOW)
 		{
 			BinauralSource.Id = ResonanceAudioApi->CreateSoundObjectSource(vraudio::RenderingMode::kBinauralLowQuality);
-			UE_LOG(LogResonanceAudio, Log, TEXT("ResonanceAudioSpatializer::OnInitSource: BINAURAL_LOW mode chosen."));
+			UE_LOG(LogResonanceAudio, Verbose, TEXT("ResonanceAudioSpatializer::OnInitSource: BINAURAL_LOW mode chosen."));
 		}
 		else if (GetDefault<UResonanceAudioSettings>()->QualityMode == ERaQualityMode::BINAURAL_MEDIUM)
 		{
 			BinauralSource.Id = ResonanceAudioApi->CreateSoundObjectSource(vraudio::RenderingMode::kBinauralMediumQuality);
-			UE_LOG(LogResonanceAudio, Log, TEXT("ResonanceAudioSpatializer::OnInitSource: BINAURAL_MEDIUM mode chosen."));
+			UE_LOG(LogResonanceAudio, Verbose, TEXT("ResonanceAudioSpatializer::OnInitSource: BINAURAL_MEDIUM mode chosen."));
 		}
 		else if (GetDefault<UResonanceAudioSettings>()->QualityMode == ERaQualityMode::BINAURAL_HIGH)
 		{
 			BinauralSource.Id = ResonanceAudioApi->CreateSoundObjectSource(vraudio::RenderingMode::kBinauralHighQuality);
-			UE_LOG(LogResonanceAudio, Log, TEXT("ResonanceAudioSpatializer::OnInitSource: BINAURAL_HIGH mode chosen."));
+			UE_LOG(LogResonanceAudio, Verbose, TEXT("ResonanceAudioSpatializer::OnInitSource: BINAURAL_HIGH mode chosen."));
 		}
 		else
 		{
-			UE_LOG(LogResonanceAudio, Error, TEXT("ResonanceAudioSpatializer::OnInitSource: Unknown quality mode!"));
+			UE_LOG(LogResonanceAudio, Verbose, TEXT("ResonanceAudioSpatializer::OnInitSource: Unknown quality mode!"));
 		}
 
 		// Set initial sound source directivity.
