Reducing log level of resonance

diff --git a/Engine/Plugins/Runtime/ResonanceAudio/Source/ResonanceAudioEditor/Private/ResonanceAudioEditorModule.h b/Engine/Plugins/Runtime/ResonanceAudio/Source/ResonanceAudioEditor/Private/ResonanceAudioEditorModule.h
--- a/Engine/Plugins/Runtime/ResonanceAudio/Source/ResonanceAudioEditor/Private/ResonanceAudioEditorModule.h
+++ b/Engine/Plugins/Runtime/ResonanceAudio/Source/ResonanceAudioEditor/Private/ResonanceAudioEditorModule.h
@@ -7,7 +7,7 @@
 #include "IResonanceAudioEditorModule.h"
 #include "Framework/MultiBox/MultiBoxExtender.h"
 
-DECLARE_LOG_CATEGORY_EXTERN(LogResonanceAudioEditor, All, All);
+DECLARE_LOG_CATEGORY_EXTERN(LogResonanceAudioEditor, Warning, All);
 
 class FSlateStyleSet;
 
