Reducing log level of non-editor resonance

diff --git a/Engine/Plugins/Runtime/ResonanceAudio/Source/ResonanceAudio/Private/ResonanceAudioEnums.h b/Engine/Plugins/Runtime/ResonanceAudio/Source/ResonanceAudio/Private/ResonanceAudioEnums.h
--- a/Engine/Plugins/Runtime/ResonanceAudio/Source/ResonanceAudio/Private/ResonanceAudioEnums.h
+++ b/Engine/Plugins/Runtime/ResonanceAudio/Source/ResonanceAudio/Private/ResonanceAudioEnums.h
@@ -7,7 +7,7 @@
 #include "CoreMinimal.h"
 #include "UObject/ObjectMacros.h"
 
-DECLARE_LOG_CATEGORY_EXTERN(LogResonanceAudio, Log, All);
+DECLARE_LOG_CATEGORY_EXTERN(LogResonanceAudio, Warning, All);
 
 UENUM(BlueprintType)
 enum class ERaQualityMode : uint8
