Fix crash in shutdown of resonance

diff --git a/Engine/Plugins/Runtime/ResonanceAudio/Source/ResonanceAudio/Private/ResonanceAudioModule.cpp b/Engine/Plugins/Runtime/ResonanceAudio/Source/ResonanceAudio/Private/ResonanceAudioModule.cpp
--- a/Engine/Plugins/Runtime/ResonanceAudio/Source/ResonanceAudio/Private/ResonanceAudioModule.cpp
+++ b/Engine/Plugins/Runtime/ResonanceAudio/Source/ResonanceAudio/Private/ResonanceAudioModule.cpp
@@ -60,12 +60,6 @@ namespace ResonanceAudio
 
 	void FResonanceAudioModule::ShutdownModule()
 	{
-		if (GlobalSpatializationSourceSettings)
-		{
-			GlobalSpatializationSourceSettings->RemoveFromRoot();
-			GlobalSpatializationSourceSettings = nullptr;
-		}
-
 		check(bModuleInitialized == true);
 		bModuleInitialized = false;
 	}
