[NO NON-PII COMMIT MESSAGE]

diff --git a/Engine/Plugins/Runtime/ResonanceAudio/Source/ResonanceAudio/Private/ResonanceAudioSettings.h b/Engine/Plugins/Runtime/ResonanceAudio/Source/ResonanceAudio/Private/ResonanceAudioSettings.h
--- a/Engine/Plugins/Runtime/ResonanceAudio/Source/ResonanceAudio/Private/ResonanceAudioSettings.h
+++ b/Engine/Plugins/Runtime/ResonanceAudio/Source/ResonanceAudio/Private/ResonanceAudioSettings.h
@@ -25,6 +25,6 @@ public:
 	FSoftObjectPath GlobalReverbPreset;
 
 	// Default settings for global source settings
-	UPROPERTY(GlobalConfig, EditAnywhere, Category = General, meta = (AllowedClasses = "ResonanceAudioSourcePluginPreset"))
+	UPROPERTY(GlobalConfig, EditAnywhere, Category = General, meta = (AllowedClasses = "ResonanceAudioSpatializationSourceSettings"))
 	FSoftObjectPath GlobalSourcePreset;
 };
