Removing static initialization of google resonance data for platfroms that don't have it enabled. If a platform doesn't define RESONANCE_SPAT_ENABLED=1 we fall back to native panning.

diff --git a/Engine/Plugins/Runtime/ResonanceAudio/Source/ResonanceAudio/Private/ResonanceAudioModule.cpp b/Engine/Plugins/Runtime/ResonanceAudio/Source/ResonanceAudio/Private/ResonanceAudioModule.cpp
--- a/Engine/Plugins/Runtime/ResonanceAudio/Source/ResonanceAudio/Private/ResonanceAudioModule.cpp
+++ b/Engine/Plugins/Runtime/ResonanceAudio/Source/ResonanceAudio/Private/ResonanceAudioModule.cpp
@@ -31,6 +31,7 @@ namespace ResonanceAudio
 	{
 		check(bModuleInitialized == false);
 
+#if RESONANCE_SPAT_ENABLED
 		bModuleInitialized = true;
 
 		// Register the Resonance Audio plugin factories.
@@ -56,16 +57,20 @@ namespace ResonanceAudio
 				GlobalSpatializationSourceSettings->AddToRoot();
 			}
 		}
+#endif // #if RESONANCE_SPAT_ENABLED
 	}
 
 	void FResonanceAudioModule::ShutdownModule()
 	{
+#if RESONANCE_SPAT_ENABLED
 		check(bModuleInitialized == true);
 		bModuleInitialized = false;
+#endif // #if RESONANCE_SPAT_ENABLED
 	}
 
 	IAudioPluginFactory* FResonanceAudioModule::GetPluginFactory(EAudioPlugin PluginType)
 	{
+#if RESONANCE_SPAT_ENABLED
 		switch (PluginType)
 		{
 		case EAudioPlugin::SPATIALIZATION:
@@ -76,6 +81,9 @@ namespace ResonanceAudio
 			break;
 		default:
 			return nullptr;
+#else
+		return nullptr;
+#endif // #if RESONANCE_SPAT_ENABLED
 		}
 	}
 
diff --git a/Engine/Plugins/Runtime/ResonanceAudio/Source/ResonanceAudio/Private/ResonanceAudioPluginListener.cpp b/Engine/Plugins/Runtime/ResonanceAudio/Source/ResonanceAudio/Private/ResonanceAudioPluginListener.cpp
--- a/Engine/Plugins/Runtime/ResonanceAudio/Source/ResonanceAudio/Private/ResonanceAudioPluginListener.cpp
+++ b/Engine/Plugins/Runtime/ResonanceAudio/Source/ResonanceAudio/Private/ResonanceAudioPluginListener.cpp
@@ -43,7 +43,7 @@ namespace ResonanceAudio
 		ResonanceAudioApi = CreateResonanceAudioApi(ResonanceAudioModule->GetResonanceAudioDynamicLibraryHandle(), 2 /* num channels */, FramesPerBuffer, SampleRate);
 		if (ResonanceAudioApi == nullptr)
 		{
-			UE_LOG(LogResonanceAudio, Error, TEXT("Failed to initialize Resonance Audio API"));
+			UE_LOG(LogResonanceAudio, Log, TEXT("Failed to initialize Resonance Audio API"));
 			return;
 		}
 
@@ -74,7 +74,7 @@ namespace ResonanceAudio
 	void FResonanceAudioPluginListener::OnListenerUpdated(FAudioDevice* AudioDevice, const int32 ViewportIndex, const FTransform& ListenerTransform, const float InDeltaSeconds)
 	{
 		if (ResonanceAudioApi == nullptr) {
-			UE_LOG(LogResonanceAudio, Error, TEXT("Resonance Audio API found"));
+			UE_LOG(LogResonanceAudio, Error, TEXT("Resonance Audio API not loaded"));
 			return;
 		}
 		else
diff --git a/Engine/Plugins/Runtime/ResonanceAudio/Source/ResonanceAudio/ResonanceAudio.Build.cs b/Engine/Plugins/Runtime/ResonanceAudio/Source/ResonanceAudio/ResonanceAudio.Build.cs
--- a/Engine/Plugins/Runtime/ResonanceAudio/Source/ResonanceAudio/ResonanceAudio.Build.cs
+++ b/Engine/Plugins/Runtime/ResonanceAudio/Source/ResonanceAudio/ResonanceAudio.Build.cs
@@ -88,6 +88,16 @@ namespace UnrealBuildTool.Rules
                 PrivateDefinitions.Add("EIGEN_HAS_CXX11_MATH=0");
             }
 
+			// whitelist platforms  to use the plugin in the first place
+            if(Target.Platform == UnrealTargetPlatform.Win64 || Target.Platform == UnrealTargetPlatform.PS4 || Target.Platform == UnrealTargetPlatform.XboxOne)
+            {
+                PublicDefinitions.Add("RESONANCE_SPAT_ENABLED=1");
+            }
+            else
+            {
+                PublicDefinitions.Add("RESONANCE_SPAT_ENABLED=0");
+            }
+
             //Embree support:
             if (Target.Platform == UnrealTargetPlatform.Win64)
             {
