Fix for CIS build error in ResonanceAudioModule

diff --git a/Engine/Plugins/Runtime/ResonanceAudio/Source/ResonanceAudio/Private/ResonanceAudioModule.cpp b/Engine/Plugins/Runtime/ResonanceAudio/Source/ResonanceAudio/Private/ResonanceAudioModule.cpp
--- a/Engine/Plugins/Runtime/ResonanceAudio/Source/ResonanceAudio/Private/ResonanceAudioModule.cpp
+++ b/Engine/Plugins/Runtime/ResonanceAudio/Source/ResonanceAudio/Private/ResonanceAudioModule.cpp
@@ -70,7 +70,7 @@ namespace ResonanceAudio
 
 	IAudioPluginFactory* FResonanceAudioModule::GetPluginFactory(EAudioPlugin PluginType)
 	{
-#if RESONANCE_SPAT_ENABLED
+#if !RESONANCE_SPAT_ENABLED
 		switch (PluginType)
 		{
 		case EAudioPlugin::SPATIALIZATION:
@@ -81,10 +81,10 @@ namespace ResonanceAudio
 			break;
 		default:
 			return nullptr;
+	}
 #else
 		return nullptr;
 #endif // #if RESONANCE_SPAT_ENABLED
-		}
 	}
 
 	void FResonanceAudioModule::RegisterAudioDevice(FAudioDevice* AudioDeviceHandle)
