Adding cvar to reduce resonance quality dynamically.

diff --git a/Engine/Plugins/Runtime/ResonanceAudio/Source/ResonanceAudio/Private/ResonanceAudioSpatialization.cpp b/Engine/Plugins/Runtime/ResonanceAudio/Source/ResonanceAudio/Private/ResonanceAudioSpatialization.cpp
--- a/Engine/Plugins/Runtime/ResonanceAudio/Source/ResonanceAudio/Private/ResonanceAudioSpatialization.cpp
+++ b/Engine/Plugins/Runtime/ResonanceAudio/Source/ResonanceAudio/Private/ResonanceAudioSpatialization.cpp
@@ -3,6 +3,15 @@
 #include "ResonanceAudioModule.h"
 #include "ResonanceAudioSettings.h"
 
+
+static int32 ResonanceQualityOverrideCVar = 0;
+FAutoConsoleVariableRef CVarResonanceQualityOverride(
+	TEXT("au.resonance.quality"),
+	ResonanceQualityOverrideCVar,
+	TEXT("Override the quality of resonance sound sources. Will not increase quality levels. The quality used will be min of the quality in the resonance source settings and this override.\n")
+	TEXT("0: Quality is not overridden, 1: Stereo Panning, 2: Low Quality, 3: Medium Quality, 4: High Quality"),
+	ECVF_Default);
+
 namespace ResonanceAudio {
 
 	/************************************************************/
@@ -72,31 +81,47 @@ namespace ResonanceAudio {
 			return;
 		}
 
-		if (SpatializationSettings[SourceId]->SpatializationMethod == ERaSpatializationMethod::STEREO_PANNING || GetDefault<UResonanceAudioSettings>()->QualityMode == ERaQualityMode::STEREO_PANNING)
-		{
-			BinauralSource.Id = ResonanceAudioApi->CreateSoundObjectSource(vraudio::RenderingMode::kStereoPanning);
-			UE_LOG(LogResonanceAudio, Verbose, TEXT("ResonanceAudioSpatializer::OnInitSource: STEREO_PANNING mode chosen."));
-		}
-		else if (GetDefault<UResonanceAudioSettings>()->QualityMode == ERaQualityMode::BINAURAL_LOW)
-		{
-			BinauralSource.Id = ResonanceAudioApi->CreateSoundObjectSource(vraudio::RenderingMode::kBinauralLowQuality);
-			UE_LOG(LogResonanceAudio, Verbose, TEXT("ResonanceAudioSpatializer::OnInitSource: BINAURAL_LOW mode chosen."));
-		}
-		else if (GetDefault<UResonanceAudioSettings>()->QualityMode == ERaQualityMode::BINAURAL_MEDIUM)
-		{
-			BinauralSource.Id = ResonanceAudioApi->CreateSoundObjectSource(vraudio::RenderingMode::kBinauralMediumQuality);
-			UE_LOG(LogResonanceAudio, Verbose, TEXT("ResonanceAudioSpatializer::OnInitSource: BINAURAL_MEDIUM mode chosen."));
-		}
-		else if (GetDefault<UResonanceAudioSettings>()->QualityMode == ERaQualityMode::BINAURAL_HIGH)
+		vraudio::RenderingMode ResonanceRenderingMode = vraudio::RenderingMode::kStereoPanning;
+
+		if (SpatializationSettings[SourceId]->SpatializationMethod != ERaSpatializationMethod::STEREO_PANNING)
 		{
-			BinauralSource.Id = ResonanceAudioApi->CreateSoundObjectSource(vraudio::RenderingMode::kBinauralHighQuality);
-			UE_LOG(LogResonanceAudio, Verbose, TEXT("ResonanceAudioSpatializer::OnInitSource: BINAURAL_HIGH mode chosen."));
+			switch (GetDefault<UResonanceAudioSettings>()->QualityMode)
+			{
+			case ERaQualityMode::STEREO_PANNING:
+				ResonanceRenderingMode = vraudio::RenderingMode::kStereoPanning;
+				UE_LOG(LogResonanceAudio, Verbose, TEXT("ResonanceAudioSpatializer::OnInitSource: STEREO_PANNING mode chosen."));
+				break;
+
+			case ERaQualityMode::BINAURAL_LOW:
+				ResonanceRenderingMode = vraudio::RenderingMode::kBinauralLowQuality;
+				UE_LOG(LogResonanceAudio, Verbose, TEXT("ResonanceAudioSpatializer::OnInitSource: BINAURAL_LOW mode chosen."));
+				break;
+
+			case ERaQualityMode::BINAURAL_MEDIUM:
+				ResonanceRenderingMode = vraudio::RenderingMode::kBinauralMediumQuality;
+				UE_LOG(LogResonanceAudio, Verbose, TEXT("ResonanceAudioSpatializer::OnInitSource: BINAURAL_MEDIUM mode chosen."));
+				break;
+
+			case ERaQualityMode::BINAURAL_HIGH:
+				ResonanceRenderingMode = vraudio::RenderingMode::kBinauralHighQuality;
+				UE_LOG(LogResonanceAudio, Verbose, TEXT("ResonanceAudioSpatializer::OnInitSource: BINAURAL_HIGH mode chosen."));
+				break;
+
+			default:
+				UE_LOG(LogResonanceAudio, Error, TEXT("ResonanceAudioSpatializer::OnInitSource: Unknown quality mode!"));
+				break;
+			}
 		}
-		else
+
+		// Only let cvar values between 0 and vraudio::RenderingMode::kRoomEffectsOnly + 1 affect the rendering mode
+		// Note that 0 is treated as non-overridden, but we want to map to the enumeration values which are 
+		if (ResonanceQualityOverrideCVar > 0 && ResonanceQualityOverrideCVar < ((int32)(vraudio::RenderingMode::kRoomEffectsOnly) + 1))
 		{
-			UE_LOG(LogResonanceAudio, Verbose, TEXT("ResonanceAudioSpatializer::OnInitSource: Unknown quality mode!"));
+			ResonanceRenderingMode = (vraudio::RenderingMode)FMath::Min((int32)ResonanceRenderingMode, ResonanceQualityOverrideCVar - 1);
 		}
 
+		BinauralSource.Id = ResonanceAudioApi->CreateSoundObjectSource(ResonanceRenderingMode);
+
 		// Set initial sound source directivity.
 		BinauralSource.Pattern = SpatializationSettings[SourceId]->Pattern;
 		BinauralSource.Sharpness = SpatializationSettings[SourceId]->Sharpness;
