Adding logging for google resonance. Extra logging can be enabled on the fly by using au.ExtraResonanceLogging 1

diff --git a/Engine/Plugins/Runtime/ResonanceAudio/Source/ResonanceAudio/Private/ResonanceAudioModule.cpp b/Engine/Plugins/Runtime/ResonanceAudio/Source/ResonanceAudio/Private/ResonanceAudioModule.cpp
--- a/Engine/Plugins/Runtime/ResonanceAudio/Source/ResonanceAudio/Private/ResonanceAudioModule.cpp
+++ b/Engine/Plugins/Runtime/ResonanceAudio/Source/ResonanceAudio/Private/ResonanceAudioModule.cpp
@@ -65,6 +65,7 @@ namespace ResonanceAudio
 #if RESONANCE_SPAT_ENABLED
 		check(bModuleInitialized == true);
 		bModuleInitialized = false;
+		UE_LOG(LogResonanceAudio, Log, TEXT("Resonance Audio Module is Shutdown"));
 #endif // #if RESONANCE_SPAT_ENABLED
 	}
 
@@ -100,6 +101,7 @@ namespace ResonanceAudio
 	void FResonanceAudioModule::UnregisterAudioDevice(FAudioDevice* AudioDeviceHandle)
 	{
 		RegisteredAudioDevices.Remove(AudioDeviceHandle);
+		UE_LOG(LogResonanceAudio, Log, TEXT("Audio Device unregistered from Resonance"));
 	}
 
 	UResonanceAudioSpatializationSourceSettings* FResonanceAudioModule::GetGlobalSpatializationSourceSettings()
diff --git a/Engine/Plugins/Runtime/ResonanceAudio/Source/ResonanceAudio/Private/ResonanceAudioPluginListener.cpp b/Engine/Plugins/Runtime/ResonanceAudio/Source/ResonanceAudio/Private/ResonanceAudioPluginListener.cpp
--- a/Engine/Plugins/Runtime/ResonanceAudio/Source/ResonanceAudio/Private/ResonanceAudioPluginListener.cpp
+++ b/Engine/Plugins/Runtime/ResonanceAudio/Source/ResonanceAudio/Private/ResonanceAudioPluginListener.cpp
@@ -43,7 +43,7 @@ namespace ResonanceAudio
 		ResonanceAudioApi = CreateResonanceAudioApi(ResonanceAudioModule->GetResonanceAudioDynamicLibraryHandle(), 2 /* num channels */, FramesPerBuffer, SampleRate);
 		if (ResonanceAudioApi == nullptr)
 		{
-			UE_LOG(LogResonanceAudio, Log, TEXT("Failed to initialize Resonance Audio API"));
+			UE_LOG(LogResonanceAudio, Error, TEXT("Failed to initialize Resonance Audio API"));
 			return;
 		}
 
@@ -69,6 +69,8 @@ namespace ResonanceAudio
 		{
 			AmbisonicsPtr->SetResonanceAudioApi(ResonanceAudioApi);
 		}
+
+		UE_LOG(LogResonanceAudio, Log, TEXT("Resonance Audio Listener is initialized"));
 	}
 
 	void FResonanceAudioPluginListener::OnListenerUpdated(FAudioDevice* AudioDevice, const int32 ViewportIndex, const FTransform& ListenerTransform, const float InDeltaSeconds)
@@ -83,6 +85,13 @@ namespace ResonanceAudio
 			ResonanceAudioApi->SetHeadPosition(ConvertedPosition.X, ConvertedPosition.Y, ConvertedPosition.Z);
 			const FQuat ConvertedRotation = ConvertToResonanceAudioRotation(ListenerTransform.GetRotation());
 			ResonanceAudioApi->SetHeadRotation(ConvertedRotation.X, ConvertedRotation.Y, ConvertedRotation.Z, ConvertedRotation.W);
+
+
+			static IConsoleVariable* ExtraBinauralLoggingCVar = IConsoleManager::Get().FindConsoleVariable(TEXT("au.ExtraResonanceLogging"));
+			if (ExtraBinauralLoggingCVar && ExtraBinauralLoggingCVar->GetBool())
+			{
+				UE_LOG(LogResonanceAudio, Warning, TEXT("Set listener position (X,Y,Z) to: (%f, %f, %f)"), ConvertedPosition.X, ConvertedPosition.Y, ConvertedPosition.Z);
+			}
 		}
 	}
 
@@ -92,6 +101,8 @@ namespace ResonanceAudio
 		{
 			ResonanceAudioModule->UnregisterAudioDevice(AudioDevice);
 		}
+
+		UE_LOG(LogResonanceAudio, Log, TEXT("Resonance Audio Listener is shutdown"));
 	}
 
 	void FResonanceAudioPluginListener::OnTick(UWorld* InWorld, const int32 ViewportIndex, const FTransform& ListenerTransform, const float InDeltaSeconds)
@@ -121,6 +132,7 @@ namespace ResonanceAudio
 			else
 			{
 				ReverbPtr->SetPreset(nullptr);
+				UE_LOG(LogResonanceAudio, Log, TEXT("Set reverb preset to nullptr"));
 			}
 		}
 	}
diff --git a/Engine/Plugins/Runtime/ResonanceAudio/Source/ResonanceAudio/Private/ResonanceAudioSpatialization.cpp b/Engine/Plugins/Runtime/ResonanceAudio/Source/ResonanceAudio/Private/ResonanceAudioSpatialization.cpp
--- a/Engine/Plugins/Runtime/ResonanceAudio/Source/ResonanceAudio/Private/ResonanceAudioSpatialization.cpp
+++ b/Engine/Plugins/Runtime/ResonanceAudio/Source/ResonanceAudio/Private/ResonanceAudioSpatialization.cpp
@@ -12,6 +12,15 @@ FAutoConsoleVariableRef CVarResonanceQualityOverride(
 	TEXT("0: Quality is not overridden, 1: Stereo Panning, 2: Low Quality, 3: Medium Quality, 4: High Quality"),
 	ECVF_Default);
 
+
+static int32 ExtraResonanceLoggingCVar = 0;
+FAutoConsoleVariableRef CVarExtraResonanceLogging(
+	TEXT("au.ExtraResonanceLogging"),
+	ExtraResonanceLoggingCVar,
+	TEXT("If non-zero, will log extra information about the state of Resonance HRTF processing.\n")
+	TEXT("0: Disable, >0: Enable"),
+	ECVF_Default);
+
 namespace ResonanceAudio {
 
 	/************************************************************/
@@ -42,6 +51,8 @@ namespace ResonanceAudio {
 		SpatializationSettings.Init(nullptr, InitializationParams.NumSources);
 
 		bIsInitialized = true;
+
+		UE_LOG(LogResonanceAudio, Log, TEXT("Google Resonance Audio Spatialization is initialized."));
 	}
 
 	bool FResonanceAudioSpatialization::IsSpatializationEffectInitialized() const
@@ -147,6 +158,12 @@ namespace ResonanceAudio {
 			UE_LOG(LogResonanceAudio, Error, TEXT("ResonanceAudioSpatializer::OnInitSource: Undefined distance roll-off model!"));
 			break;
 		}
+
+
+		if (ExtraResonanceLoggingCVar)
+		{
+			UE_LOG(LogResonanceAudio, Warning, TEXT("Source initialized (Our source ID: %i) (Binaural SourceId)."), SourceId, BinauralSource.Id);
+		}
 	}
 
 	void FResonanceAudioSpatialization::OnReleaseSource(const uint32 SourceId)
@@ -162,6 +179,11 @@ namespace ResonanceAudio {
 			ResonanceAudioApi->DestroySource(BinauralSource.Id);
 			BinauralSource.Id = RA_INVALID_SOURCE_ID;
 		}
+
+		if (ExtraResonanceLoggingCVar)
+		{
+			UE_LOG(LogResonanceAudio, Warning, TEXT("Destroying Source (Our Source ID: %i)(Binaural Source ID: %i)"), SourceId, BinauralSource.Id);
+		}
 	}
 
 	void FResonanceAudioSpatialization::ProcessAudio(const FAudioPluginSourceInputData& InputData, FAudioPluginSourceOutputData& OutputData)
@@ -199,6 +221,15 @@ namespace ResonanceAudio {
 
 			// Add source buffer to process.
 			ResonanceAudioApi->SetInterleavedBuffer(BinauralSource.Id, InputData.AudioBuffer->GetData(), InputData.NumChannels, InputData.AudioBuffer->Num() / InputData.NumChannels);
+
+			// optional heartbeat to make sure Resonance is processing audio
+			static const int32 NumBuffersBetweenLogs = 500;
+			static int32 ResonanceAudioBuffersProcessecdCounter = 0;
+
+			if (ExtraResonanceLoggingCVar && (ResonanceAudioBuffersProcessecdCounter++ % NumBuffersBetweenLogs) == 0)
+			{
+				UE_LOG(LogResonanceAudio, Warning, TEXT("FResonanceAudioSpatialization::ProcessAudio() has been called %i times"), NumBuffersBetweenLogs);
+			}
 		}
 	}
 
