change logging level for LogResonanceAudio

diff --git a/Engine/Plugins/Runtime/ResonanceAudio/Source/ResonanceAudio/Private/ResonanceAudioEnums.h b/Engine/Plugins/Runtime/ResonanceAudio/Source/ResonanceAudio/Private/ResonanceAudioEnums.h
--- a/Engine/Plugins/Runtime/ResonanceAudio/Source/ResonanceAudio/Private/ResonanceAudioEnums.h
+++ b/Engine/Plugins/Runtime/ResonanceAudio/Source/ResonanceAudio/Private/ResonanceAudioEnums.h
@@ -7,7 +7,7 @@
 #include "CoreMinimal.h"
 #include "UObject/ObjectMacros.h"
 
-DECLARE_LOG_CATEGORY_EXTERN(LogResonanceAudio, Warning, All);
+DECLARE_LOG_CATEGORY_EXTERN(LogResonanceAudio, Display, All);
 
 UENUM(BlueprintType)
 enum class ERaQualityMode : uint8
diff --git a/Engine/Plugins/Runtime/ResonanceAudio/Source/ResonanceAudio/Private/ResonanceAudioPluginListener.cpp b/Engine/Plugins/Runtime/ResonanceAudio/Source/ResonanceAudio/Private/ResonanceAudioPluginListener.cpp
--- a/Engine/Plugins/Runtime/ResonanceAudio/Source/ResonanceAudio/Private/ResonanceAudioPluginListener.cpp
+++ b/Engine/Plugins/Runtime/ResonanceAudio/Source/ResonanceAudio/Private/ResonanceAudioPluginListener.cpp
@@ -70,7 +70,7 @@ namespace ResonanceAudio
 			AmbisonicsPtr->SetResonanceAudioApi(ResonanceAudioApi);
 		}
 
-		UE_LOG(LogResonanceAudio, Log, TEXT("Resonance Audio Listener is initialized"));
+		UE_LOG(LogResonanceAudio, Display, TEXT("Resonance Audio Listener is initialized"));
 	}
 
 	void FResonanceAudioPluginListener::OnListenerUpdated(FAudioDevice* AudioDevice, const int32 ViewportIndex, const FTransform& ListenerTransform, const float InDeltaSeconds)
@@ -85,13 +85,6 @@ namespace ResonanceAudio
 			ResonanceAudioApi->SetHeadPosition(ConvertedPosition.X, ConvertedPosition.Y, ConvertedPosition.Z);
 			const FQuat ConvertedRotation = ConvertToResonanceAudioRotation(ListenerTransform.GetRotation());
 			ResonanceAudioApi->SetHeadRotation(ConvertedRotation.X, ConvertedRotation.Y, ConvertedRotation.Z, ConvertedRotation.W);
-
-
-			static IConsoleVariable* ExtraBinauralLoggingCVar = IConsoleManager::Get().FindConsoleVariable(TEXT("au.ExtraResonanceLogging"));
-			if (ExtraBinauralLoggingCVar && ExtraBinauralLoggingCVar->GetBool())
-			{
-				UE_LOG(LogResonanceAudio, Warning, TEXT("Set listener position (X,Y,Z) to: (%f, %f, %f)"), ConvertedPosition.X, ConvertedPosition.Y, ConvertedPosition.Z);
-			}
 		}
 	}
 
@@ -102,7 +95,7 @@ namespace ResonanceAudio
 			ResonanceAudioModule->UnregisterAudioDevice(AudioDevice);
 		}
 
-		UE_LOG(LogResonanceAudio, Log, TEXT("Resonance Audio Listener is shutdown"));
+		UE_LOG(LogResonanceAudio, Display, TEXT("Resonance Audio Listener is shutdown"));
 	}
 
 	void FResonanceAudioPluginListener::OnTick(UWorld* InWorld, const int32 ViewportIndex, const FTransform& ListenerTransform, const float InDeltaSeconds)
@@ -132,7 +125,7 @@ namespace ResonanceAudio
 			else
 			{
 				ReverbPtr->SetPreset(nullptr);
-				UE_LOG(LogResonanceAudio, Log, TEXT("Set reverb preset to nullptr"));
+				UE_LOG(LogResonanceAudio, Display, TEXT("Set reverb preset to nullptr"));
 			}
 		}
 	}
diff --git a/Engine/Plugins/Runtime/ResonanceAudio/Source/ResonanceAudio/Private/ResonanceAudioSpatialization.cpp b/Engine/Plugins/Runtime/ResonanceAudio/Source/ResonanceAudio/Private/ResonanceAudioSpatialization.cpp
--- a/Engine/Plugins/Runtime/ResonanceAudio/Source/ResonanceAudio/Private/ResonanceAudioSpatialization.cpp
+++ b/Engine/Plugins/Runtime/ResonanceAudio/Source/ResonanceAudio/Private/ResonanceAudioSpatialization.cpp
@@ -162,7 +162,7 @@ namespace ResonanceAudio {
 
 		if (ExtraResonanceLoggingCVar)
 		{
-			UE_LOG(LogResonanceAudio, Warning, TEXT("Source initialized (Our source ID: %i) (Binaural SourceId)."), SourceId, BinauralSource.Id);
+			UE_LOG(LogResonanceAudio, Warning, TEXT("Source initialized (Our source ID: %i) (Binaural SourceId: %i)."), SourceId, BinauralSource.Id);
 		}
 	}
 
