Putting spammy Google Resonance log behind cvar

diff --git a/Engine/Plugins/Runtime/ResonanceAudio/Source/ResonanceAudio/Private/ResonanceAudioPluginListener.cpp b/Engine/Plugins/Runtime/ResonanceAudio/Source/ResonanceAudio/Private/ResonanceAudioPluginListener.cpp
--- a/Engine/Plugins/Runtime/ResonanceAudio/Source/ResonanceAudio/Private/ResonanceAudioPluginListener.cpp
+++ b/Engine/Plugins/Runtime/ResonanceAudio/Source/ResonanceAudio/Private/ResonanceAudioPluginListener.cpp
@@ -100,6 +100,7 @@ namespace ResonanceAudio
 
 	void FResonanceAudioPluginListener::OnTick(UWorld* InWorld, const int32 ViewportIndex, const FTransform& ListenerTransform, const float InDeltaSeconds)
 	{
+		static IConsoleVariable* ExtraBinauralLoggingCVar = IConsoleManager::Get().FindConsoleVariable(TEXT("au.ExtraResonanceLogging"));
 		if (ReverbPtr != nullptr && InWorld->AudioVolumes.Num() > 0)
 		{
 			AAudioVolume* CurrentVolume = InWorld->GetAudioSettings(ListenerTransform.GetLocation(), nullptr, nullptr);
@@ -125,7 +126,10 @@ namespace ResonanceAudio
 			else
 			{
 				ReverbPtr->SetPreset(nullptr);
-				UE_LOG(LogResonanceAudio, Display, TEXT("Set reverb preset to nullptr"));
+				if (ExtraBinauralLoggingCVar && ExtraBinauralLoggingCVar->GetBool())
+				{	
+					UE_LOG(LogResonanceAudio, Display, TEXT("Set reverb preset to nullptr"));
+				}
 			}
 		}
 	}
