Integrating fix for spammy logging

diff --git a/Engine/Plugins/Runtime/ResonanceAudio/Source/ResonanceAudio/Private/ResonanceAudioPluginListener.cpp b/Engine/Plugins/Runtime/ResonanceAudio/Source/ResonanceAudio/Private/ResonanceAudioPluginListener.cpp
--- a/Engine/Plugins/Runtime/ResonanceAudio/Source/ResonanceAudio/Private/ResonanceAudioPluginListener.cpp
+++ b/Engine/Plugins/Runtime/ResonanceAudio/Source/ResonanceAudio/Private/ResonanceAudioPluginListener.cpp
@@ -100,7 +100,6 @@ namespace ResonanceAudio
 
 	void FResonanceAudioPluginListener::OnTick(UWorld* InWorld, const int32 ViewportIndex, const FTransform& ListenerTransform, const float InDeltaSeconds)
 	{
-		static IConsoleVariable* ExtraBinauralLoggingCVar = IConsoleManager::Get().FindConsoleVariable(TEXT("au.ExtraResonanceLogging"));
 		if (ReverbPtr != nullptr && InWorld->AudioVolumes.Num() > 0)
 		{
 			AAudioVolume* CurrentVolume = InWorld->GetAudioSettings(ListenerTransform.GetLocation(), nullptr, nullptr);
@@ -126,10 +125,7 @@ namespace ResonanceAudio
 			else
 			{
 				ReverbPtr->SetPreset(nullptr);
-				if (ExtraBinauralLoggingCVar && ExtraBinauralLoggingCVar->GetBool())
-				{	
-					UE_LOG(LogResonanceAudio, Display, TEXT("Set reverb preset to nullptr"));
-				}
+				UE_LOG(LogResonanceAudio, Verbose, TEXT("Set reverb preset to nullptr"));
 			}
 		}
 	}
