[NO NON-PII COMMIT MESSAGE]

diff --git a/Engine/Plugins/Runtime/ResonanceAudio/Source/ResonanceAudio/Private/ResonanceAudioLibrary/resonance_audio/utils/lockless_task_queue.cc b/Engine/Plugins/Runtime/ResonanceAudio/Source/ResonanceAudio/Private/ResonanceAudioLibrary/resonance_audio/utils/lockless_task_queue.cc
--- a/Engine/Plugins/Runtime/ResonanceAudio/Source/ResonanceAudio/Private/ResonanceAudioLibrary/resonance_audio/utils/lockless_task_queue.cc
+++ b/Engine/Plugins/Runtime/ResonanceAudio/Source/ResonanceAudio/Private/ResonanceAudioLibrary/resonance_audio/utils/lockless_task_queue.cc
@@ -34,6 +34,8 @@ _Pragma("clang diagnostic ignored \"-Wshadow\"")
 
 #include "base/logging.h"
 
+#include <mutex>
+
 namespace vraudio {
 
 LocklessTaskQueue::LocklessTaskQueue(size_t max_tasks) {
@@ -95,6 +97,8 @@ LocklessTaskQueue::Node* LocklessTaskQueue::PopNodeFromList(
 }
 
 void LocklessTaskQueue::ProcessTaskList(Node* list_head, bool execute) {
+  std::lock_guard<std::mutex> scope_lock(temp_tasks_mutex_);
+
   Node* node_itr = list_head;
   while (node_itr != nullptr) {
     Node* next_node_ptr = node_itr->next;
@@ -117,6 +121,7 @@ void LocklessTaskQueue::ProcessTaskList(Node* list_head, bool execute) {
 }
 
 void LocklessTaskQueue::Init(size_t num_nodes) {
+  std::lock_guard<std::mutex> scope_lock(temp_tasks_mutex_);
   nodes_.resize(num_nodes);
   temp_tasks_.reserve(num_nodes);
 
diff --git a/Engine/Plugins/Runtime/ResonanceAudio/Source/ResonanceAudio/Private/ResonanceAudioLibrary/resonance_audio/utils/lockless_task_queue.h b/Engine/Plugins/Runtime/ResonanceAudio/Source/ResonanceAudio/Private/ResonanceAudioLibrary/resonance_audio/utils/lockless_task_queue.h
--- a/Engine/Plugins/Runtime/ResonanceAudio/Source/ResonanceAudio/Private/ResonanceAudioLibrary/resonance_audio/utils/lockless_task_queue.h
+++ b/Engine/Plugins/Runtime/ResonanceAudio/Source/ResonanceAudio/Private/ResonanceAudioLibrary/resonance_audio/utils/lockless_task_queue.h
@@ -36,6 +36,7 @@ _Pragma("clang diagnostic ignored \"-Wshadow\"")
 #include <stddef.h>
 
 #include <atomic>
+#include <mutex>
 #include <functional>
 #include <vector>
 
@@ -117,6 +118,9 @@ class LocklessTaskQueue {
   // Temporary vector to hold |Task|s in order to execute them in reverse order
   // (FIFO, instead of LIFO).
   std::vector<Task> temp_tasks_;
+
+  // this mutex guarantees that temp_tasks_ is not mutated while being iterated on.
+  std::mutex temp_tasks_mutex_;
 };
 
 }  // namespace vraudio
