Speculative fix for Resonance crash

diff --git a/Engine/Plugins/Runtime/ResonanceAudio/Source/ResonanceAudio/Private/ResonanceAudioLibrary/resonance_audio/utils/lockless_task_queue.cc b/Engine/Plugins/Runtime/ResonanceAudio/Source/ResonanceAudio/Private/ResonanceAudioLibrary/resonance_audio/utils/lockless_task_queue.cc
--- a/Engine/Plugins/Runtime/ResonanceAudio/Source/ResonanceAudio/Private/ResonanceAudioLibrary/resonance_audio/utils/lockless_task_queue.cc
+++ b/Engine/Plugins/Runtime/ResonanceAudio/Source/ResonanceAudio/Private/ResonanceAudioLibrary/resonance_audio/utils/lockless_task_queue.cc
@@ -38,7 +38,7 @@ _Pragma("clang diagnostic ignored \"-Wshadow\"")
 
 namespace vraudio {
 
-LocklessTaskQueue::LocklessTaskQueue(size_t max_tasks) {
+LocklessTaskQueue::LocklessTaskQueue(size_t max_tasks) : max_tasks_(max_tasks) {
   CHECK_GT(max_tasks, 0U);
   Init(max_tasks);
 }
@@ -102,7 +102,12 @@ void LocklessTaskQueue::ProcessTaskList(Node* list_head, bool execute) {
   Node* node_itr = list_head;
   while (node_itr != nullptr) {
     Node* next_node_ptr = node_itr->next;
-    temp_tasks_.emplace_back(std::move(node_itr->task));
+    if (temp_tasks_.size() < max_tasks_) {
+      temp_tasks_.emplace_back(std::move(node_itr->task));
+    } else{
+      LOG(WARNING) << "temp_tasks_ vector size is larger than max_tasks_, dropping all remaining tasks";
+    }
+
     node_itr->task = nullptr;
     PushNodeToList(&free_list_head_, node_itr);
     node_itr = next_node_ptr;
@@ -118,6 +123,11 @@ void LocklessTaskQueue::ProcessTaskList(Node* list_head, bool execute) {
     }
   }
   temp_tasks_.clear();
+  static bool bHasWarned = false;
+  if (bHasWarned && max_tasks_ && !temp_tasks_.capacity()) {
+    LOG(WARNING) << "std::vector::clear() on this platform has zero'd out the temp_tasks_ capacity, (this will cause re-allocation in every ProcessTaskList call)";
+    bHasWarned = false;
+  }
 }
 
 void LocklessTaskQueue::Init(size_t num_nodes) {
diff --git a/Engine/Plugins/Runtime/ResonanceAudio/Source/ResonanceAudio/Private/ResonanceAudioLibrary/resonance_audio/utils/lockless_task_queue.h b/Engine/Plugins/Runtime/ResonanceAudio/Source/ResonanceAudio/Private/ResonanceAudioLibrary/resonance_audio/utils/lockless_task_queue.h
--- a/Engine/Plugins/Runtime/ResonanceAudio/Source/ResonanceAudio/Private/ResonanceAudioLibrary/resonance_audio/utils/lockless_task_queue.h
+++ b/Engine/Plugins/Runtime/ResonanceAudio/Source/ResonanceAudio/Private/ResonanceAudioLibrary/resonance_audio/utils/lockless_task_queue.h
@@ -121,6 +121,7 @@ class LocklessTaskQueue {
 
   // this mutex guarantees that temp_tasks_ is not mutated while being iterated on.
   std::mutex temp_tasks_mutex_;
+  size_t max_tasks_;
 };
 
 }  // namespace vraudio
