[NO NON-PII COMMIT MESSAGE]

diff --git a/Engine/Plugins/Runtime/ResonanceAudio/Source/ResonanceAudio/Private/ResonanceAudioLibrary/resonance_audio/dsp/sh_hrir_creator.cc b/Engine/Plugins/Runtime/ResonanceAudio/Source/ResonanceAudio/Private/ResonanceAudioLibrary/resonance_audio/dsp/sh_hrir_creator.cc
--- a/Engine/Plugins/Runtime/ResonanceAudio/Source/ResonanceAudio/Private/ResonanceAudioLibrary/resonance_audio/dsp/sh_hrir_creator.cc
+++ b/Engine/Plugins/Runtime/ResonanceAudio/Source/ResonanceAudio/Private/ResonanceAudioLibrary/resonance_audio/dsp/sh_hrir_creator.cc
@@ -61,11 +61,12 @@ std::unique_ptr<AudioBuffer> CreateShHrirsFromAssets(
     const std::string& filename, int target_sample_rate_hz,
     Resampler* resampler) {
   // Read SH HRIR from asset store.
-  sadie::HrtfAssets hrtf_assets;
-  std::unique_ptr<std::string> sh_hrir_data = hrtf_assets.GetFile(filename);
-  CHECK_NOTNULL(sh_hrir_data.get());
-  std::istringstream wav_data_stream(*sh_hrir_data);
-  std::unique_ptr<const Wav> wav = Wav::CreateOrNull(&wav_data_stream);
+  size_t sh_hrir_data_size = 0;
+  const unsigned char* sh_hrir_data = sadie::HrtfAssets::GetFile(filename, sh_hrir_data_size);
+  CHECK_NOTNULL(sh_hrir_data);
+
+  std::unique_ptr<const Wav> wav = Wav::CreateOrNull(sh_hrir_data, sh_hrir_data_size);
+  CHECK_NOTNULL(wav.get());
   return CreateShHrirsFromWav(*wav, target_sample_rate_hz, resampler);
 }
 
diff --git a/Engine/Plugins/Runtime/ResonanceAudio/Source/ResonanceAudio/Private/ResonanceAudioLibrary/resonance_audio/utils/wav.cc b/Engine/Plugins/Runtime/ResonanceAudio/Source/ResonanceAudio/Private/ResonanceAudioLibrary/resonance_audio/utils/wav.cc
--- a/Engine/Plugins/Runtime/ResonanceAudio/Source/ResonanceAudio/Private/ResonanceAudioLibrary/resonance_audio/utils/wav.cc
+++ b/Engine/Plugins/Runtime/ResonanceAudio/Source/ResonanceAudio/Private/ResonanceAudioLibrary/resonance_audio/utils/wav.cc
@@ -34,8 +34,8 @@ Wav::Wav(size_t num_channels, int sample_rate,
 
 Wav::~Wav() {}
 
-std::unique_ptr<const Wav> Wav::CreateOrNull(std::istream* binary_stream) {
-  WavReader wav_reader(binary_stream);
+std::unique_ptr<const Wav> Wav::CreateOrNull(const unsigned char* byte_array, size_t array_size) {
+  WavReader wav_reader(byte_array, array_size);
   const size_t num_total_samples = wav_reader.GetNumTotalSamples();
   if (!wav_reader.IsHeaderValid() || num_total_samples == 0) {
     return nullptr;
diff --git a/Engine/Plugins/Runtime/ResonanceAudio/Source/ResonanceAudio/Private/ResonanceAudioLibrary/resonance_audio/utils/wav.h b/Engine/Plugins/Runtime/ResonanceAudio/Source/ResonanceAudio/Private/ResonanceAudioLibrary/resonance_audio/utils/wav.h
--- a/Engine/Plugins/Runtime/ResonanceAudio/Source/ResonanceAudio/Private/ResonanceAudioLibrary/resonance_audio/utils/wav.h
+++ b/Engine/Plugins/Runtime/ResonanceAudio/Source/ResonanceAudio/Private/ResonanceAudioLibrary/resonance_audio/utils/wav.h
@@ -28,8 +28,8 @@ class Wav {
  public:
   ~Wav();
 
-  // Reads a RIFF WAVE from an opened binary stream.
-  static std::unique_ptr<const Wav> CreateOrNull(std::istream* binary_stream);
+  // Reads a RIFF WAVE from a byte array.
+  static std::unique_ptr<const Wav> CreateOrNull(const unsigned char* byte_array, size_t array_size);
 
   // Returns reference to interleaved samples.
   const std::vector<int16_t>& interleaved_samples() const {
diff --git a/Engine/Plugins/Runtime/ResonanceAudio/Source/ResonanceAudio/Private/ResonanceAudioLibrary/resonance_audio/utils/wav_reader.cc b/Engine/Plugins/Runtime/ResonanceAudio/Source/ResonanceAudio/Private/ResonanceAudioLibrary/resonance_audio/utils/wav_reader.cc
--- a/Engine/Plugins/Runtime/ResonanceAudio/Source/ResonanceAudio/Private/ResonanceAudioLibrary/resonance_audio/utils/wav_reader.cc
+++ b/Engine/Plugins/Runtime/ResonanceAudio/Source/ResonanceAudio/Private/ResonanceAudioLibrary/resonance_audio/utils/wav_reader.cc
@@ -64,8 +64,8 @@ static const uint16 kExtensibleWavFormat = 0xfffe;
 static const uint16 kPcmFormat = 0x1;
 }  // namespace
 
-WavReader::WavReader(std::istream* binary_stream)
-    : binary_stream_(CHECK_NOTNULL(binary_stream)),
+WavReader::WavReader(const unsigned char* byte_array, const size_t num_bytes)
+    : source(byte_array, num_bytes),
       num_channels_(0),
       sample_rate_hz_(-1),
       num_total_samples_(0),
@@ -75,11 +75,15 @@ WavReader::WavReader(std::istream* binary_stream)
 }
 
 size_t WavReader::ReadBinaryDataFromStream(void* target_ptr, size_t size) {
-  if (!binary_stream_->good()) {
+  // if we are out of bounds, return 0.
+  if (source.current_byte_index + size > source.num_bytes) {
     return 0;
   }
-  binary_stream_->read(static_cast<char*>(target_ptr), size);
-  return static_cast<size_t>(binary_stream_->gcount());
+
+  //binary_stream_->read(static_cast<char*>(target_ptr), size);
+  std::memcpy(target_ptr, &source.byte_array[source.current_byte_index], size);
+  source.current_byte_index += size;
+  return size;
 }
 
 bool WavReader::ParseHeader() {
@@ -155,11 +159,7 @@ bool WavReader::ParseHeader() {
     return false;
   }
 
-  int64 current_position = binary_stream_->tellg();
-  if (current_position < 0) {
-    return false;
-  }
-  pcm_offset_bytes_ = static_cast<uint64>(current_position);
+  pcm_offset_bytes_ = static_cast<uint64>(source.current_byte_index);
   return true;
 }
 
@@ -182,11 +182,14 @@ int64 WavReader::SeekToFrame(const uint64 frame_position) {
   if (frame_position <= (num_total_samples_ / num_channels_)) {
     const uint64 seek_position_byte =
         pcm_offset_bytes_ + frame_position * num_channels_ * bytes_per_sample_;
-    binary_stream_->seekg(seek_position_byte, binary_stream_->beg);
+    
+	//ensure the offset position isn't beyond the end of our stream:
+	DCHECK_LT(seek_position_byte, source.num_bytes);
+	source.current_byte_index = seek_position_byte;
   }
 
   int64 binary_stream_position_byte =
-      static_cast<int64>(binary_stream_->tellg());
+      static_cast<int64>(source.current_byte_index);
   if (binary_stream_position_byte < 0) {
     // Return an error status if the actual bitstream position could not be
     // queried.
diff --git a/Engine/Plugins/Runtime/ResonanceAudio/Source/ResonanceAudio/Private/ResonanceAudioLibrary/resonance_audio/utils/wav_reader.h b/Engine/Plugins/Runtime/ResonanceAudio/Source/ResonanceAudio/Private/ResonanceAudioLibrary/resonance_audio/utils/wav_reader.h
--- a/Engine/Plugins/Runtime/ResonanceAudio/Source/ResonanceAudio/Private/ResonanceAudioLibrary/resonance_audio/utils/wav_reader.h
+++ b/Engine/Plugins/Runtime/ResonanceAudio/Source/ResonanceAudio/Private/ResonanceAudioLibrary/resonance_audio/utils/wav_reader.h
@@ -27,10 +27,12 @@ namespace vraudio {
 // Basic RIFF WAVE decoder that supports multichannel 16-bit PCM.
 class WavReader {
  public:
-  // Constructor decodes WAV header.
-  //
-  // @param binary_stream Binary input stream to read from.
-  explicit WavReader(std::istream* binary_stream);
+
+  // Constructor that takes a pointer to a raw byte array and the size of that byte stream.
+  // 
+  // @param byte_array Pointer to raw array of bytes.
+  // @param num_bytes Size of byte_array.
+  explicit WavReader(const unsigned char* byte_array, const size_t num_bytes);
 
   // True if WAV header was successfully parsed.
   bool IsHeaderValid() const;
@@ -73,8 +75,22 @@ class WavReader {
   // @return Number of bytes read.
   size_t ReadBinaryDataFromStream(void* target_ptr, size_t size);
 
-  // Binary input stream.
-  std::istream* binary_stream_;
+  struct SourceByteArray
+  {
+	  const unsigned char* byte_array;
+	  const size_t num_bytes;
+	  size_t current_byte_index;
+
+	  SourceByteArray(const unsigned char* in_byte_array, const size_t in_num_bytes)
+		  : byte_array(in_byte_array)
+		  , num_bytes(in_num_bytes)
+		  , current_byte_index(0)
+	  {
+	  }
+  };
+
+  // source we are deserializing.
+  SourceByteArray source;
 
   // Flag indicating if the WAV header was parsed successfully.
   bool init_;
diff --git a/Engine/Plugins/Runtime/ResonanceAudio/Source/ResonanceAudio/Private/ResonanceAudioModule.cpp b/Engine/Plugins/Runtime/ResonanceAudio/Source/ResonanceAudio/Private/ResonanceAudioModule.cpp
--- a/Engine/Plugins/Runtime/ResonanceAudio/Source/ResonanceAudio/Private/ResonanceAudioModule.cpp
+++ b/Engine/Plugins/Runtime/ResonanceAudio/Source/ResonanceAudio/Private/ResonanceAudioModule.cpp
@@ -30,8 +30,6 @@ namespace ResonanceAudio
 	void FResonanceAudioModule::StartupModule()
 	{
 		check(bModuleInitialized == false);
-
-#if RESONANCE_SPAT_ENABLED
 		bModuleInitialized = true;
 
 		// Register the Resonance Audio plugin factories.
@@ -57,21 +55,17 @@ namespace ResonanceAudio
 				GlobalSpatializationSourceSettings->AddToRoot();
 			}
 		}
-#endif // #if RESONANCE_SPAT_ENABLED
 	}
 
 	void FResonanceAudioModule::ShutdownModule()
 	{
-#if RESONANCE_SPAT_ENABLED
 		check(bModuleInitialized == true);
 		bModuleInitialized = false;
 		UE_LOG(LogResonanceAudio, Log, TEXT("Resonance Audio Module is Shutdown"));
-#endif // #if RESONANCE_SPAT_ENABLED
 	}
 
 	IAudioPluginFactory* FResonanceAudioModule::GetPluginFactory(EAudioPlugin PluginType)
 	{
-#if RESONANCE_SPAT_ENABLED
 		switch (PluginType)
 		{
 		case EAudioPlugin::SPATIALIZATION:
@@ -83,9 +77,6 @@ namespace ResonanceAudio
 		default:
 			return nullptr;
 		}
-#else
-		return nullptr;
-#endif // #if RESONANCE_SPAT_ENABLED
 	}
 
 	void FResonanceAudioModule::RegisterAudioDevice(FAudioDevice* AudioDeviceHandle)
diff --git a/Engine/Plugins/Runtime/ResonanceAudio/Source/ResonanceAudio/ResonanceAudio.Build.cs b/Engine/Plugins/Runtime/ResonanceAudio/Source/ResonanceAudio/ResonanceAudio.Build.cs
--- a/Engine/Plugins/Runtime/ResonanceAudio/Source/ResonanceAudio/ResonanceAudio.Build.cs
+++ b/Engine/Plugins/Runtime/ResonanceAudio/Source/ResonanceAudio/ResonanceAudio.Build.cs
@@ -87,16 +87,6 @@ public class ResonanceAudio : ModuleRules
             PrivateDefinitions.Add("EIGEN_HAS_CXX11_MATH=0");
         }
 
-		// whitelist platforms  to use the plugin in the first place
-        if(Target.Platform == UnrealTargetPlatform.Win64 || Target.Platform == UnrealTargetPlatform.PS4 || Target.Platform == UnrealTargetPlatform.XboxOne)
-        {
-            PublicDefinitions.Add("RESONANCE_SPAT_ENABLED=1");
-        }
-        else
-        {
-            PublicDefinitions.Add("RESONANCE_SPAT_ENABLED=0");
-        }
-
         //Embree support:
         if (Target.Platform == UnrealTargetPlatform.Win64)
         {
