#ResonanceAudio Include fix

diff --git a/Engine/Plugins/Runtime/ResonanceAudio/Source/ResonanceAudio/Private/ResonanceAudioLibrary/resonance_audio/utils/wav_reader.cc b/Engine/Plugins/Runtime/ResonanceAudio/Source/ResonanceAudio/Private/ResonanceAudioLibrary/resonance_audio/utils/wav_reader.cc
--- a/Engine/Plugins/Runtime/ResonanceAudio/Source/ResonanceAudio/Private/ResonanceAudioLibrary/resonance_audio/utils/wav_reader.cc
+++ b/Engine/Plugins/Runtime/ResonanceAudio/Source/ResonanceAudio/Private/ResonanceAudioLibrary/resonance_audio/utils/wav_reader.cc
@@ -18,6 +18,7 @@ limitations under the License.
 
 #include <algorithm>
 #include <string>
+#include <memory>
 
 #include "base/integral_types.h"
 #include "base/logging.h"
