Change std::memcpy to memcpy to fix android compile

diff --git a/Engine/Plugins/Runtime/ResonanceAudio/Source/ResonanceAudio/Private/ResonanceAudioLibrary/resonance_audio/utils/wav_reader.cc b/Engine/Plugins/Runtime/ResonanceAudio/Source/ResonanceAudio/Private/ResonanceAudioLibrary/resonance_audio/utils/wav_reader.cc
--- a/Engine/Plugins/Runtime/ResonanceAudio/Source/ResonanceAudio/Private/ResonanceAudioLibrary/resonance_audio/utils/wav_reader.cc
+++ b/Engine/Plugins/Runtime/ResonanceAudio/Source/ResonanceAudio/Private/ResonanceAudioLibrary/resonance_audio/utils/wav_reader.cc
@@ -83,7 +83,7 @@ size_t WavReader::ReadBinaryDataFromStream(void* target_ptr, size_t size) {
   }
 
   //binary_stream_->read(static_cast<char*>(target_ptr), size);
-  std::memcpy(target_ptr, &source.byte_array[source.current_byte_index], size);
+  memcpy(target_ptr, &source.byte_array[source.current_byte_index], size);
   source.current_byte_index += size;
   return size;
 }
