[NO NON-PII COMMIT MESSAGE]

diff --git a/Engine/Plugins/Runtime/ResonanceAudio/Source/ResonanceAudio/Private/ResonanceAudioLibrary/resonance_audio/base/constants_and_types.h b/Engine/Plugins/Runtime/ResonanceAudio/Source/ResonanceAudio/Private/ResonanceAudioLibrary/resonance_audio/base/constants_and_types.h
--- a/Engine/Plugins/Runtime/ResonanceAudio/Source/ResonanceAudio/Private/ResonanceAudioLibrary/resonance_audio/base/constants_and_types.h
+++ b/Engine/Plugins/Runtime/ResonanceAudio/Source/ResonanceAudio/Private/ResonanceAudioLibrary/resonance_audio/base/constants_and_types.h
@@ -20,11 +20,9 @@ limitations under the License.
 #include <cmath>
 #include <string>  // for size_t
 
-#include "CoreMinimal.h"
-
 #ifndef M_PI
 #define DEFINED_M_PI
-#define M_PI PI
+#define M_PI 3.14159265358979323846
 #endif
 
 namespace vraudio {
diff --git a/Engine/Plugins/Runtime/ResonanceAudio/Source/ResonanceAudio/Private/ResonanceAudioLibrary/resonance_audio/dsp/filter_coefficient_generators.cc b/Engine/Plugins/Runtime/ResonanceAudio/Source/ResonanceAudio/Private/ResonanceAudioLibrary/resonance_audio/dsp/filter_coefficient_generators.cc
--- a/Engine/Plugins/Runtime/ResonanceAudio/Source/ResonanceAudio/Private/ResonanceAudioLibrary/resonance_audio/dsp/filter_coefficient_generators.cc
+++ b/Engine/Plugins/Runtime/ResonanceAudio/Source/ResonanceAudio/Private/ResonanceAudioLibrary/resonance_audio/dsp/filter_coefficient_generators.cc
@@ -21,11 +21,9 @@ limitations under the License.
 
 #include "base/constants_and_types.h"
 
-#include "CoreMinimal.h"
-
 #ifndef M_PI
 #define DEFINED_M_PI
-#define M_PI PI
+#define M_PI 3.14159265358979323846
 #endif
 
 namespace vraudio {
