[NO NON-PII COMMIT MESSAGE]

diff --git a/Engine/Plugins/Runtime/ResonanceAudio/Source/ResonanceAudio/Private/ResonanceAudioLibrary/resonance_audio/base/logging.h b/Engine/Plugins/Runtime/ResonanceAudio/Source/ResonanceAudio/Private/ResonanceAudioLibrary/resonance_audio/base/logging.h
--- a/Engine/Plugins/Runtime/ResonanceAudio/Source/ResonanceAudio/Private/ResonanceAudioLibrary/resonance_audio/base/logging.h
+++ b/Engine/Plugins/Runtime/ResonanceAudio/Source/ResonanceAudio/Private/ResonanceAudioLibrary/resonance_audio/base/logging.h
@@ -52,6 +52,7 @@ class NullLogger {
 
 
 #define DISABLE_CHECKS 1
+#define INCLUDE_UNUSED_CODE 0
 
 #if DISABLE_CHECKS
 #define DCHECK(expr) 
diff --git a/Engine/Plugins/Runtime/ResonanceAudio/Source/ResonanceAudio/Private/ResonanceAudioLibrary/resonance_audio/dsp/fir_filter.cc b/Engine/Plugins/Runtime/ResonanceAudio/Source/ResonanceAudio/Private/ResonanceAudioLibrary/resonance_audio/dsp/fir_filter.cc
--- a/Engine/Plugins/Runtime/ResonanceAudio/Source/ResonanceAudio/Private/ResonanceAudioLibrary/resonance_audio/dsp/fir_filter.cc
+++ b/Engine/Plugins/Runtime/ResonanceAudio/Source/ResonanceAudio/Private/ResonanceAudioLibrary/resonance_audio/dsp/fir_filter.cc
@@ -89,7 +89,11 @@ void FirFilter::Process(const AudioBuffer::Channel& input,
   // other to store data that needed to be copied as it straddled a four float
   // boundary.
   SimdVector input_use;
+
+#if defined(SIMD_SSE)
   SimdVector input_split;
+#endif
+
   for (size_t input_position = num_filter_chunks_;
        input_position < num_input_chunks + num_filter_chunks_;
        ++input_position) {
diff --git a/Engine/Plugins/Runtime/ResonanceAudio/Source/ResonanceAudio/Private/ResonanceAudioLibrary/resonance_audio/dsp/multi_channel_iir.cc b/Engine/Plugins/Runtime/ResonanceAudio/Source/ResonanceAudio/Private/ResonanceAudioLibrary/resonance_audio/dsp/multi_channel_iir.cc
--- a/Engine/Plugins/Runtime/ResonanceAudio/Source/ResonanceAudio/Private/ResonanceAudioLibrary/resonance_audio/dsp/multi_channel_iir.cc
+++ b/Engine/Plugins/Runtime/ResonanceAudio/Source/ResonanceAudio/Private/ResonanceAudioLibrary/resonance_audio/dsp/multi_channel_iir.cc
@@ -142,7 +142,9 @@ void MultiChannelIir::Process(AudioBuffer::Channel* interleaved_buffer) {
 MultiChannelIir::MultiChannelIir(size_t num_channels, size_t frames_per_buffer,
                                  size_t num_coefficients)
     : num_channels_(num_channels),
+#if !DISABLE_CHECKS
       frames_per_buffer_(frames_per_buffer),
+#endif
       num_coefficients_(num_coefficients),
       delay_line_front_(0),
       numerator_(kNumMonoChannels, num_coefficients_ * num_channels_),
diff --git a/Engine/Plugins/Runtime/ResonanceAudio/Source/ResonanceAudio/Private/ResonanceAudioLibrary/resonance_audio/dsp/multi_channel_iir.h b/Engine/Plugins/Runtime/ResonanceAudio/Source/ResonanceAudio/Private/ResonanceAudioLibrary/resonance_audio/dsp/multi_channel_iir.h
--- a/Engine/Plugins/Runtime/ResonanceAudio/Source/ResonanceAudio/Private/ResonanceAudioLibrary/resonance_audio/dsp/multi_channel_iir.h
+++ b/Engine/Plugins/Runtime/ResonanceAudio/Source/ResonanceAudio/Private/ResonanceAudioLibrary/resonance_audio/dsp/multi_channel_iir.h
@@ -64,8 +64,10 @@ class MultiChannelIir {
   // Number of channels in each input buffer.
   const size_t num_channels_;
 
+#if !DISABLE_CHECKS
   // Number of frames in each input buffer.
   const size_t frames_per_buffer_;
+#endif
 
   // Number of coefficients in the numerator and denominator polynomials.
   const size_t num_coefficients_;
diff --git a/Engine/Plugins/Runtime/ResonanceAudio/Source/ResonanceAudio/Private/ResonanceAudioLibrary/resonance_audio/geometrical_acoustics/path_tracer.cc b/Engine/Plugins/Runtime/ResonanceAudio/Source/ResonanceAudio/Private/ResonanceAudioLibrary/resonance_audio/geometrical_acoustics/path_tracer.cc
--- a/Engine/Plugins/Runtime/ResonanceAudio/Source/ResonanceAudio/Private/ResonanceAudioLibrary/resonance_audio/geometrical_acoustics/path_tracer.cc
+++ b/Engine/Plugins/Runtime/ResonanceAudio/Source/ResonanceAudio/Private/ResonanceAudioLibrary/resonance_audio/geometrical_acoustics/path_tracer.cc
@@ -24,6 +24,7 @@ limitations under the License.
 
 namespace vraudio {
 
+#if INCLUDE_UNUSED_CODE
 std::vector<Path> PathTracer::TracePaths(const AcousticSource& source,
                                          size_t min_num_rays, size_t max_depth,
                                          float energy_threshold) {
@@ -92,5 +93,6 @@ std::vector<Path> PathTracer::TracePaths(const AcousticSource& source,
       });
   return paths;
 }
+#endif
 
 }  // namespace vraudio
diff --git a/Engine/Plugins/Runtime/ResonanceAudio/Source/ResonanceAudio/Private/ResonanceAudioLibrary/resonance_audio/geometrical_acoustics/path_tracer.h b/Engine/Plugins/Runtime/ResonanceAudio/Source/ResonanceAudio/Private/ResonanceAudioLibrary/resonance_audio/geometrical_acoustics/path_tracer.h
--- a/Engine/Plugins/Runtime/ResonanceAudio/Source/ResonanceAudio/Private/ResonanceAudioLibrary/resonance_audio/geometrical_acoustics/path_tracer.h
+++ b/Engine/Plugins/Runtime/ResonanceAudio/Source/ResonanceAudio/Private/ResonanceAudioLibrary/resonance_audio/geometrical_acoustics/path_tracer.h
@@ -36,6 +36,7 @@ class PathTracer {
       : scene_manager_(scene_manager) {}
   ~PathTracer() {}
 
+#if INCLUDE_UNUSED_CODE
   // Traces sound propagation paths from a source.
   //
   // @param source Source from which paths are traced.
@@ -48,6 +49,7 @@ class PathTracer {
   std::vector<Path> TracePaths(const AcousticSource& source,
                                size_t min_num_rays, size_t max_depth,
                                float energy_threshold);
+#endif
 
  private:
   // Scene manager.
diff --git a/Engine/Plugins/Runtime/ResonanceAudio/Source/ResonanceAudio/Private/ResonanceAudioLibrary/resonance_audio/geometrical_acoustics/reflection_kernel.cc b/Engine/Plugins/Runtime/ResonanceAudio/Source/ResonanceAudio/Private/ResonanceAudioLibrary/resonance_audio/geometrical_acoustics/reflection_kernel.cc
--- a/Engine/Plugins/Runtime/ResonanceAudio/Source/ResonanceAudio/Private/ResonanceAudioLibrary/resonance_audio/geometrical_acoustics/reflection_kernel.cc
+++ b/Engine/Plugins/Runtime/ResonanceAudio/Source/ResonanceAudio/Private/ResonanceAudioLibrary/resonance_audio/geometrical_acoustics/reflection_kernel.cc
@@ -23,6 +23,7 @@ namespace vraudio {
 
 using Eigen::Vector3f;
 
+#if INCLUDE_UNUSED_CODE
 AcousticRay ReflectionKernel::Reflect(const AcousticRay& incident_ray) const {
   // This function currently uses |random_number_generator| (which returns
   // uniformly distributed numbers) to naively sample 1D and 2D points.
@@ -69,6 +70,7 @@ AcousticRay ReflectionKernel::Reflect(const AcousticRay& incident_ray) const {
                      new_energies, reflected_ray_type,
                      reflected_ray_prior_distance);
 }
+#endif
 
 void ReflectionKernel::ReflectDiffuseRain(
     const AcousticRay& incident_ray, const AcousticRay& reference_reflected_ray,
diff --git a/Engine/Plugins/Runtime/ResonanceAudio/Source/ResonanceAudio/Private/ResonanceAudioLibrary/resonance_audio/geometrical_acoustics/reflection_kernel.h b/Engine/Plugins/Runtime/ResonanceAudio/Source/ResonanceAudio/Private/ResonanceAudioLibrary/resonance_audio/geometrical_acoustics/reflection_kernel.h
--- a/Engine/Plugins/Runtime/ResonanceAudio/Source/ResonanceAudio/Private/ResonanceAudioLibrary/resonance_audio/geometrical_acoustics/reflection_kernel.h
+++ b/Engine/Plugins/Runtime/ResonanceAudio/Source/ResonanceAudio/Private/ResonanceAudioLibrary/resonance_audio/geometrical_acoustics/reflection_kernel.h
@@ -61,6 +61,7 @@ class ReflectionKernel {
     CHECK(scattering_coefficient_ >= 0.0f && scattering_coefficient_ <= 1.0f);
   }
 
+#if INCLUDE_UNUSED_CODE
   // Reflects a ray.
   //
   // @param incident_ray Incident (incoming) ray.
@@ -68,6 +69,7 @@ class ReflectionKernel {
   //     with direction and energy calculated based on the type of reflection
   //     and the incident ray.
   AcousticRay Reflect(const AcousticRay& incident_ray) const;
+#endif
 
   // Reflects an incident ray to create a diffuse-rain ray, whose |direction|
   // is from the reflection point to the listener position, and whose |t_far|
diff --git a/Engine/Plugins/Runtime/ResonanceAudio/Source/ResonanceAudio/Private/ResonanceAudioLibrary/resonance_audio/geometrical_acoustics/sampling.h b/Engine/Plugins/Runtime/ResonanceAudio/Source/ResonanceAudio/Private/ResonanceAudioLibrary/resonance_audio/geometrical_acoustics/sampling.h
--- a/Engine/Plugins/Runtime/ResonanceAudio/Source/ResonanceAudio/Private/ResonanceAudioLibrary/resonance_audio/geometrical_acoustics/sampling.h
+++ b/Engine/Plugins/Runtime/ResonanceAudio/Source/ResonanceAudio/Private/ResonanceAudioLibrary/resonance_audio/geometrical_acoustics/sampling.h
@@ -73,11 +73,11 @@ inline Eigen::Vector3f StratifiedSampleSphere(float u, float v,
                              (cell_y + v) * cell_width);
 }
 
+#if INCLUDE_UNUSED_CODE
 // Samples a vector from a unit hemisphere according to the cosine-weighted
 // distribution given two random variables whose values are in [0, 1). The
 // hemisphere lies on a plane whose normal is assumed to be on the +z direction.
 //
-
 static Eigen::Vector3f CosineSampleHemisphere(float u, float v) {
   DCHECK(u >= 0.0f && u <= 1.0f);
   DCHECK(v >= 0.0f && v <= 1.0f);
@@ -103,6 +103,7 @@ static Eigen::Vector3f CosineSampleHemisphere(
                                             unit_normal) *
          local_vector;
 }
+#endif
 
 // The probability density function (PDF) that a vector is in a particular
 // direction if the vector is sampled from a cosine-weigted distribution over
diff --git a/Engine/Plugins/Runtime/ResonanceAudio/Source/ResonanceAudio/Private/ResonanceAudioLibrary/resonance_audio/utils/buffer_unpartitioner.cc b/Engine/Plugins/Runtime/ResonanceAudio/Source/ResonanceAudio/Private/ResonanceAudioLibrary/resonance_audio/utils/buffer_unpartitioner.cc
--- a/Engine/Plugins/Runtime/ResonanceAudio/Source/ResonanceAudio/Private/ResonanceAudioLibrary/resonance_audio/utils/buffer_unpartitioner.cc
+++ b/Engine/Plugins/Runtime/ResonanceAudio/Source/ResonanceAudio/Private/ResonanceAudioLibrary/resonance_audio/utils/buffer_unpartitioner.cc
@@ -29,7 +29,10 @@ namespace vraudio {
 BufferUnpartitioner::BufferUnpartitioner(size_t num_channels,
                                          size_t frames_per_buffer,
                                          GetBufferCallback buffer_callback)
-    : num_channels_(num_channels),
+    :
+#if !DISABLE_CHECKS
+	num_channels_(num_channels),
+#endif
       frames_per_buffer_(frames_per_buffer),
       buffer_callback_(std::move(buffer_callback)),
       current_input_buffer_(nullptr),
diff --git a/Engine/Plugins/Runtime/ResonanceAudio/Source/ResonanceAudio/Private/ResonanceAudioLibrary/resonance_audio/utils/buffer_unpartitioner.h b/Engine/Plugins/Runtime/ResonanceAudio/Source/ResonanceAudio/Private/ResonanceAudioLibrary/resonance_audio/utils/buffer_unpartitioner.h
--- a/Engine/Plugins/Runtime/ResonanceAudio/Source/ResonanceAudio/Private/ResonanceAudioLibrary/resonance_audio/utils/buffer_unpartitioner.h
+++ b/Engine/Plugins/Runtime/ResonanceAudio/Source/ResonanceAudio/Private/ResonanceAudioLibrary/resonance_audio/utils/buffer_unpartitioner.h
@@ -114,8 +114,10 @@ class BufferUnpartitioner {
   size_t GetBufferTemplated(BufferType buffer, size_t num_channels,
                             size_t num_frames);
 
+#if !DISABLE_CHECKS
   // Number of channels in output buffers.
   const size_t num_channels_;
+#endif
 
   // Number of frames per buffer in output buffers.
   const size_t frames_per_buffer_;
diff --git a/Engine/Plugins/Runtime/ResonanceAudio/Source/ResonanceAudio/Private/ResonanceAudioLibrary/resonance_audio/utils/planar_interleaved_conversion.cc b/Engine/Plugins/Runtime/ResonanceAudio/Source/ResonanceAudio/Private/ResonanceAudioLibrary/resonance_audio/utils/planar_interleaved_conversion.cc
--- a/Engine/Plugins/Runtime/ResonanceAudio/Source/ResonanceAudio/Private/ResonanceAudioLibrary/resonance_audio/utils/planar_interleaved_conversion.cc
+++ b/Engine/Plugins/Runtime/ResonanceAudio/Source/ResonanceAudio/Private/ResonanceAudioLibrary/resonance_audio/utils/planar_interleaved_conversion.cc
@@ -38,6 +38,8 @@ void ConvertInterleavedToPlanarTemplated(
     const std::vector<size_t>* channel_map, OutputType output_buffer,
     size_t num_output_frames, size_t num_output_channels,
     size_t output_offset_frames, size_t num_frames_to_copy) {
+
+#if !DISABLE_CHECKS
   DCHECK_GE(num_input_frames, input_offset_frames);
   const size_t max_num_input_frames = num_input_frames - input_offset_frames;
   DCHECK_GE(num_output_frames, output_offset_frames);
@@ -50,6 +52,7 @@ void ConvertInterleavedToPlanarTemplated(
   } else {
     DCHECK_GE(channel_map->size(), num_output_channels);
   }
+#endif
 
   InputType interleaved_buffer_with_offset =
       interleaved_buffer + input_offset_frames * num_input_channels;
@@ -92,6 +95,8 @@ void ConvertPlanarToPlanarTemplated(
     PlanarOutputType planar_output_ptrs, size_t num_output_frames,
     size_t num_output_channels, size_t output_offset_frames,
     size_t num_frames_convert_and_copy) {
+
+#if !DISABLE_CHECKS
   DCHECK_GE(num_input_frames, input_offset_frames);
   const size_t max_num_input_frames = num_input_frames - input_offset_frames;
   DCHECK_GE(num_output_frames, output_offset_frames);
@@ -104,6 +109,7 @@ void ConvertPlanarToPlanarTemplated(
   } else {
     DCHECK_GE(channel_map->size(), num_output_channels);
   }
+#endif
 
   for (size_t channel = 0; channel < num_output_channels; ++channel) {
     const size_t input_channel =
@@ -121,6 +127,8 @@ void ConvertPlanarToInterleavedTemplated(
     size_t input_offset_frames, InterleavedOutputType interleaved_output_ptr,
     size_t num_output_frames, size_t num_output_channels,
     size_t output_offset_frames, size_t num_frames_convert_and_copy) {
+
+#if !DISABLE_CHECKS
   DCHECK(interleaved_output_ptr);
   DCHECK_GE(num_input_frames, input_offset_frames);
   const size_t max_num_input_frames = num_input_frames - input_offset_frames;
@@ -129,6 +137,7 @@ void ConvertPlanarToInterleavedTemplated(
   DCHECK_GE(max_num_input_frames, num_frames_convert_and_copy);
   DCHECK_GE(max_num_output_frames, num_frames_convert_and_copy);
   DCHECK_EQ(num_input_channels, num_output_channels);
+#endif
 
   InterleavedOutputType interleaved_output_ptr_with_offset =
       interleaved_output_ptr + output_offset_frames * num_output_channels;
