[NO NON-PII COMMIT MESSAGE]

diff --git a/Engine/Plugins/Runtime/ResonanceAudio/Source/ResonanceAudio/Private/ResonanceAudioLibrary/resonance_audio/base/aligned_allocator.h b/Engine/Plugins/Runtime/ResonanceAudio/Source/ResonanceAudio/Private/ResonanceAudioLibrary/resonance_audio/base/aligned_allocator.h
--- a/Engine/Plugins/Runtime/ResonanceAudio/Source/ResonanceAudio/Private/ResonanceAudioLibrary/resonance_audio/base/aligned_allocator.h
+++ b/Engine/Plugins/Runtime/ResonanceAudio/Source/ResonanceAudio/Private/ResonanceAudioLibrary/resonance_audio/base/aligned_allocator.h
@@ -26,6 +26,15 @@ limitations under the License.
 
 #include "base/simd_utils.h"
 
+#if defined(__clang__)
+_Pragma("clang diagnostic push")
+
+// suppress clang warnings here.
+
+#else
+__pragma(warning(push))
+__pragma(warning(disable: 4100))
+#endif
 
 namespace vraudio {
 
@@ -114,4 +123,10 @@ class AlignedAllocator : public std::allocator<Type> {
 
 }  // namespace vraudio
 
+#if defined(__clang__)
+_Pragma("clang diagnostic pop")
+#else
+__pragma(warning(pop))
+#endif
+
 #endif  // RESONANCE_AUDIO_BASE_ALIGNED_ALLOCATOR_H_
