[NO NON-PII COMMIT MESSAGE]

diff --git a/Engine/Plugins/Runtime/ResonanceAudio/Source/ResonanceAudio/Private/ResonanceAudioLibrary/resonance_audio/base/simd_utils.cc b/Engine/Plugins/Runtime/ResonanceAudio/Source/ResonanceAudio/Private/ResonanceAudioLibrary/resonance_audio/base/simd_utils.cc
--- a/Engine/Plugins/Runtime/ResonanceAudio/Source/ResonanceAudio/Private/ResonanceAudioLibrary/resonance_audio/base/simd_utils.cc
+++ b/Engine/Plugins/Runtime/ResonanceAudio/Source/ResonanceAudio/Private/ResonanceAudioLibrary/resonance_audio/base/simd_utils.cc
@@ -29,6 +29,15 @@ limitations under the License.
 #include "base/misc_math.h"
 #include "base/simd_macros.h"
 
+#if defined(__clang__)
+_Pragma("clang diagnostic push")
+
+// suppress clang warnings here.
+
+#else
+__pragma(warning(push))
+__pragma(warning(disable: 4100))
+#endif
 
 namespace vraudio {
 
@@ -1297,3 +1306,9 @@ void DeinterleaveQuad(size_t length, const float* interleaved_buffer,
 }
 
 }  // namespace vraudio
+
+#if defined(__clang__)
+_Pragma("clang diagnostic pop")
+#else
+__pragma(warning(pop))
+#endif
\ No newline at end of file
