[NO NON-PII COMMIT MESSAGE]

diff --git a/Engine/Plugins/Runtime/ResonanceAudio/Source/ResonanceAudio/Private/ResonanceAudioLibrary/resonance_audio/base/audio_buffer.cc b/Engine/Plugins/Runtime/ResonanceAudio/Source/ResonanceAudio/Private/ResonanceAudioLibrary/resonance_audio/base/audio_buffer.cc
--- a/Engine/Plugins/Runtime/ResonanceAudio/Source/ResonanceAudio/Private/ResonanceAudioLibrary/resonance_audio/base/audio_buffer.cc
+++ b/Engine/Plugins/Runtime/ResonanceAudio/Source/ResonanceAudio/Private/ResonanceAudioLibrary/resonance_audio/base/audio_buffer.cc
@@ -18,10 +18,10 @@ limitations under the License.
 
 namespace vraudio {
 
-AudioBuffer::AudioBuffer() : num_frames_(0), source_id_(kInvalidSourceId) {}
+AudioBuffer::AudioBuffer() : num_frames_(0), data_size_(0), source_id_(kInvalidSourceId) {}
 
 AudioBuffer::AudioBuffer(size_t num_channels, size_t num_frames)
-    : num_frames_(num_frames), source_id_(kInvalidSourceId) {
+    : num_frames_(num_frames), data_size_(0), source_id_(kInvalidSourceId) {
 
   InitChannelViews(num_channels);
 }
