fix for IOS compile failures

diff --git a/Engine/Plugins/Runtime/ResonanceAudio/Source/ResonanceAudio/Private/ResonanceAudioLibrary/resonance_audio/geometrical_acoustics/path_tracer.h b/Engine/Plugins/Runtime/ResonanceAudio/Source/ResonanceAudio/Private/ResonanceAudioLibrary/resonance_audio/geometrical_acoustics/path_tracer.h
--- a/Engine/Plugins/Runtime/ResonanceAudio/Source/ResonanceAudio/Private/ResonanceAudioLibrary/resonance_audio/geometrical_acoustics/path_tracer.h
+++ b/Engine/Plugins/Runtime/ResonanceAudio/Source/ResonanceAudio/Private/ResonanceAudioLibrary/resonance_audio/geometrical_acoustics/path_tracer.h
@@ -33,7 +33,10 @@ class PathTracer {
   //
   // @param scene_manager Scene manager.
   explicit PathTracer(const SceneManager& scene_manager)
-      : scene_manager_(scene_manager) {}
+#if INCLUDE_UNUSED_CODE
+      : scene_manager_(scene_manager)
+#endif
+	{}
   ~PathTracer() {}
 
 #if INCLUDE_UNUSED_CODE
@@ -53,7 +56,9 @@ class PathTracer {
 
  private:
   // Scene manager.
+#if INCLUDE_UNUSED_CODE
   const SceneManager& scene_manager_;
+#endif
 };
 
 }  // namespace vraudio
