fix for PC compile failure

diff --git a/Engine/Plugins/Runtime/ResonanceAudio/Source/ResonanceAudio/Private/ResonanceAudioLibrary/resonance_audio/geometrical_acoustics/impulse_response_computer.cc b/Engine/Plugins/Runtime/ResonanceAudio/Source/ResonanceAudio/Private/ResonanceAudioLibrary/resonance_audio/geometrical_acoustics/impulse_response_computer.cc
--- a/Engine/Plugins/Runtime/ResonanceAudio/Source/ResonanceAudio/Private/ResonanceAudioLibrary/resonance_audio/geometrical_acoustics/impulse_response_computer.cc
+++ b/Engine/Plugins/Runtime/ResonanceAudio/Source/ResonanceAudio/Private/ResonanceAudioLibrary/resonance_audio/geometrical_acoustics/impulse_response_computer.cc
@@ -33,7 +33,9 @@ ImpulseResponseComputer::ImpulseResponseComputer(
     SceneManager* scene_manager)
     : listeners_(std::move(listeners)),
       collection_kernel_(listener_sphere_radius, sampling_rate),
+#if INCLUDE_UNUSED_CODE
       scene_manager_(scene_manager),
+#endif
       finalized_(false),
       num_total_paths_(0) {
   CHECK_NOTNULL(listeners_.get());
@@ -42,6 +44,7 @@ ImpulseResponseComputer::ImpulseResponseComputer(
 
 ImpulseResponseComputer::~ImpulseResponseComputer() {}
 
+#if INCLUDE_UNUSED_CODE
 void ImpulseResponseComputer::CollectContributions(
     const std::vector<Path>& paths_batch) {
   // Do not collect anymore if finalized.
@@ -160,7 +163,8 @@ void ImpulseResponseComputer::CollectContributions(
       });
   num_total_paths_ += paths_batch.size();
 }
-
+#endif
+	
 const std::vector<AcousticListener>&
 ImpulseResponseComputer::GetFinalizedListeners() {
   DCHECK_GT(num_total_paths_, 0U);
diff --git a/Engine/Plugins/Runtime/ResonanceAudio/Source/ResonanceAudio/Private/ResonanceAudioLibrary/resonance_audio/geometrical_acoustics/impulse_response_computer.h b/Engine/Plugins/Runtime/ResonanceAudio/Source/ResonanceAudio/Private/ResonanceAudioLibrary/resonance_audio/geometrical_acoustics/impulse_response_computer.h
--- a/Engine/Plugins/Runtime/ResonanceAudio/Source/ResonanceAudio/Private/ResonanceAudioLibrary/resonance_audio/geometrical_acoustics/impulse_response_computer.h
+++ b/Engine/Plugins/Runtime/ResonanceAudio/Source/ResonanceAudio/Private/ResonanceAudioLibrary/resonance_audio/geometrical_acoustics/impulse_response_computer.h
@@ -74,12 +74,14 @@ class ImpulseResponseComputer {
       SceneManager* scene_manager);
   virtual ~ImpulseResponseComputer();
 
+#if INCLUDE_UNUSED_CODE
   // Collects contributions from a batch of paths to all listeners if
   // the collection is not finalized yet.
   //
   // @param paths_batch All sound propagation paths in a batch.
   void CollectContributions(const std::vector<Path>& paths_batch);
-
+#endif
+	
   // Finalizes the listeners and returns them. After calling this, further
   // calls to CollectContributions() have no effect.
   //
@@ -93,12 +95,14 @@ class ImpulseResponseComputer {
   // Collection kernel used to collect contributions from rays to listeners.
   CollectionKernel collection_kernel_;
 
+#if INCLUDE_UNUSED_CODE
   // Scene manager. Used to keep records of listener spheres for ray-sphere
   // intersection tests. Also used in the diffuse rain algorithm to to test
   // whether there is an un-obstructed path from a reflection point to a
   // listener.
   SceneManager* scene_manager_;
-
+#endif
+	
   // Is the collection finalized.
   bool finalized_;
 
