Explicitly truncate uint64 to size_t on 32 bit machines.

diff --git a/Engine/Plugins/Runtime/ResonanceAudio/Source/ResonanceAudio/Private/ResonanceAudioLibrary/resonance_audio/utils/wav_reader.cc b/Engine/Plugins/Runtime/ResonanceAudio/Source/ResonanceAudio/Private/ResonanceAudioLibrary/resonance_audio/utils/wav_reader.cc
--- a/Engine/Plugins/Runtime/ResonanceAudio/Source/ResonanceAudio/Private/ResonanceAudioLibrary/resonance_audio/utils/wav_reader.cc
+++ b/Engine/Plugins/Runtime/ResonanceAudio/Source/ResonanceAudio/Private/ResonanceAudioLibrary/resonance_audio/utils/wav_reader.cc
@@ -186,7 +186,11 @@ int64 WavReader::SeekToFrame(const uint64 frame_position) {
     
 	//ensure the offset position isn't beyond the end of our stream:
 	DCHECK_LT(seek_position_byte, source.num_bytes);
-	source.current_byte_index = seek_position_byte;
+
+	// for 32 bit devices, we need to truncate seek_position_byte to a 32 bit value.
+	// If seek_position_byte is less than num_bytes, we should not need to worry about
+	// overflow.
+	source.current_byte_index = static_cast<size_t>(seek_position_byte);
   }
 
   int64 binary_stream_position_byte =
