Fix uninitialized member variables in binaural_surround_renderer_impl.cc

diff --git a/Engine/Plugins/Runtime/ResonanceAudio/Source/ResonanceAudio/Private/ResonanceAudioLibrary/resonance_audio/graph/binaural_surround_renderer_impl.cc b/Engine/Plugins/Runtime/ResonanceAudio/Source/ResonanceAudio/Private/ResonanceAudioLibrary/resonance_audio/graph/binaural_surround_renderer_impl.cc
--- a/Engine/Plugins/Runtime/ResonanceAudio/Source/ResonanceAudio/Private/ResonanceAudioLibrary/resonance_audio/graph/binaural_surround_renderer_impl.cc
+++ b/Engine/Plugins/Runtime/ResonanceAudio/Source/ResonanceAudio/Private/ResonanceAudioLibrary/resonance_audio/graph/binaural_surround_renderer_impl.cc
@@ -117,8 +117,11 @@ BinauralSurroundRendererImpl::BinauralSurroundRendererImpl()
       resonance_audio_api_(nullptr),
       frames_per_buffer_(0),
       sample_rate_hz_(0),
+	  surround_format_(SurroundFormat::kInvalid),
+	  num_input_channels_(0),
       total_frames_buffered_(0),
-      num_zero_padded_frames_(0) {
+      num_zero_padded_frames_(0),
+      output_gain_(0.0f) {
 }
 
 AudioBuffer* BinauralSurroundRendererImpl::BufferPartitionerCallback(
