remove unreachable code in resonance.

diff --git a/Engine/Plugins/Runtime/ResonanceAudio/Source/ResonanceAudio/Private/ResonanceAudioLibrary/resonance_audio/graph/binaural_surround_renderer_impl.cc b/Engine/Plugins/Runtime/ResonanceAudio/Source/ResonanceAudio/Private/ResonanceAudioLibrary/resonance_audio/graph/binaural_surround_renderer_impl.cc
--- a/Engine/Plugins/Runtime/ResonanceAudio/Source/ResonanceAudio/Private/ResonanceAudioLibrary/resonance_audio/graph/binaural_surround_renderer_impl.cc
+++ b/Engine/Plugins/Runtime/ResonanceAudio/Source/ResonanceAudio/Private/ResonanceAudioLibrary/resonance_audio/graph/binaural_surround_renderer_impl.cc
@@ -162,9 +162,9 @@ size_t BinauralSurroundRendererImpl::GetExpectedNumChannelsFromSurroundFormat(
       return kNumThirdOrderAmbisonicChannels + kNumStereoChannels;
     default:
       LOG(FATAL) << "Undefined surround format mode";
-      return false;
       break;
   }
+
   return 0;
 }
 
