Removing warnings for C++17 deprecations in std::allocator

diff --git a/Engine/Plugins/Runtime/ResonanceAudio/Source/ResonanceAudio/Private/ResonanceAudioLibrary/resonance_audio/base/aligned_allocator.h b/Engine/Plugins/Runtime/ResonanceAudio/Source/ResonanceAudio/Private/ResonanceAudioLibrary/resonance_audio/base/aligned_allocator.h
--- a/Engine/Plugins/Runtime/ResonanceAudio/Source/ResonanceAudio/Private/ResonanceAudioLibrary/resonance_audio/base/aligned_allocator.h
+++ b/Engine/Plugins/Runtime/ResonanceAudio/Source/ResonanceAudio/Private/ResonanceAudioLibrary/resonance_audio/base/aligned_allocator.h
@@ -34,6 +34,7 @@ _Pragma("clang diagnostic push")
 #else
 __pragma(warning(push))
 __pragma(warning(disable: 4100))
+__pragma(warning(disable: 4996))
 #endif
 
 namespace vraudio {
@@ -81,9 +82,17 @@ void AllignedFree(PointerType mem_block_aligned) {
 template <typename Type, size_t Alignment>
 class AlignedAllocator : public std::allocator<Type> {
  public:
+
+#if __cplusplus >= 201703L
+  using Pointer = pointer;
+  using ConstPointer = const_pointer;
+  using SizeType = size_type;
+
+#else
   typedef typename std::allocator<Type>::pointer Pointer;
   typedef typename std::allocator<Type>::const_pointer ConstPointer;
   typedef typename std::allocator<Type>::size_type SizeType;
+#endif
 
   AlignedAllocator() { StaticAlignmentCheck<sizeof(Type), Alignment>(); }
 
