Speculative Fix for high-rate crash in Resonance

diff --git a/Engine/Plugins/Runtime/ResonanceAudio/Source/ResonanceAudio/Private/ResonanceAudioLibrary/resonance_audio/utils/lockless_task_queue.cc b/Engine/Plugins/Runtime/ResonanceAudio/Source/ResonanceAudio/Private/ResonanceAudioLibrary/resonance_audio/utils/lockless_task_queue.cc
--- a/Engine/Plugins/Runtime/ResonanceAudio/Source/ResonanceAudio/Private/ResonanceAudioLibrary/resonance_audio/utils/lockless_task_queue.cc
+++ b/Engine/Plugins/Runtime/ResonanceAudio/Source/ResonanceAudio/Private/ResonanceAudioLibrary/resonance_audio/utils/lockless_task_queue.cc
@@ -97,43 +97,34 @@ LocklessTaskQueue::Node* LocklessTaskQueue::PopNodeFromList(
 }
 
 void LocklessTaskQueue::ProcessTaskList(Node* list_head, bool execute) {
-  std::lock_guard<std::mutex> scope_lock(temp_tasks_mutex_);
-
   Node* node_itr = list_head;
+  int32_t num_tasks_processed = 0;
+
   while (node_itr != nullptr) {
     Node* next_node_ptr = node_itr->next;
-    if (temp_tasks_.size() < max_tasks_) {
-      temp_tasks_.emplace_back(std::move(node_itr->task));
+    if (num_tasks_processed < max_tasks_) {
+	  if (execute && node_itr->task != nullptr) {
+		node_itr->task();
+	  }
+	  ++num_tasks_processed;
     } else{
       LOG(WARNING) << "temp_tasks_ vector size is larger than max_tasks_, dropping all remaining tasks";
     }
 
-    node_itr->task = nullptr;
+	node_itr->task = Task();
     PushNodeToList(&free_list_head_, node_itr);
     node_itr = next_node_ptr;
   }
 
-  if (execute) {
-    // Execute tasks in reverse order.
-    for (std::vector<Task>::reverse_iterator task_itr = temp_tasks_.rbegin();
-         task_itr != temp_tasks_.rend(); ++task_itr) {
-      if (*task_itr != nullptr) {
-        (*task_itr)();
-      }
-    }
-  }
-  temp_tasks_.clear();
   static bool bHasWarned = false;
-  if (bHasWarned && max_tasks_ && !temp_tasks_.capacity()) {
+  if (bHasWarned && max_tasks_ && num_tasks_processed >= max_tasks_) {
     LOG(WARNING) << "std::vector::clear() on this platform has zero'd out the temp_tasks_ capacity, (this will cause re-allocation in every ProcessTaskList call)";
     bHasWarned = false;
   }
 }
 
 void LocklessTaskQueue::Init(size_t num_nodes) {
-  std::lock_guard<std::mutex> scope_lock(temp_tasks_mutex_);
   nodes_.resize(num_nodes);
-  temp_tasks_.reserve(num_nodes);
 
   // Initialize free list.
   free_list_head_ = &nodes_[0];
diff --git a/Engine/Plugins/Runtime/ResonanceAudio/Source/ResonanceAudio/Private/ResonanceAudioLibrary/resonance_audio/utils/lockless_task_queue.h b/Engine/Plugins/Runtime/ResonanceAudio/Source/ResonanceAudio/Private/ResonanceAudioLibrary/resonance_audio/utils/lockless_task_queue.h
--- a/Engine/Plugins/Runtime/ResonanceAudio/Source/ResonanceAudio/Private/ResonanceAudioLibrary/resonance_audio/utils/lockless_task_queue.h
+++ b/Engine/Plugins/Runtime/ResonanceAudio/Source/ResonanceAudio/Private/ResonanceAudioLibrary/resonance_audio/utils/lockless_task_queue.h
@@ -114,13 +114,6 @@ class LocklessTaskQueue {
 
   // Holds preallocated nodes.
   std::vector<Node> nodes_;
-
-  // Temporary vector to hold |Task|s in order to execute them in reverse order
-  // (FIFO, instead of LIFO).
-  std::vector<Task> temp_tasks_;
-
-  // this mutex guarantees that temp_tasks_ is not mutated while being iterated on.
-  std::mutex temp_tasks_mutex_;
   size_t max_tasks_;
 };
 
