Make Mac work with Resonance

diff --git a/Engine/Plugins/Runtime/ResonanceAudio/Source/ResonanceAudio/ResonanceAudio.Build.cs b/Engine/Plugins/Runtime/ResonanceAudio/Source/ResonanceAudio/ResonanceAudio.Build.cs
--- a/Engine/Plugins/Runtime/ResonanceAudio/Source/ResonanceAudio/ResonanceAudio.Build.cs
+++ b/Engine/Plugins/Runtime/ResonanceAudio/Source/ResonanceAudio/ResonanceAudio.Build.cs
@@ -104,6 +104,10 @@ public class ResonanceAudio : ModuleRules
             string SDKDir = Target.UEThirdPartySourceDirectory + "IntelEmbree/Embree2140/MacOSX/";
 
             PublicIncludePaths.Add(SDKDir + "include");
+            PublicAdditionalLibraries.Add(SDKDir + "lib/libembree.2.14.0.dylib");
+            RuntimeDependencies.Add("$(TargetOutputDir)/libembree.2.14.0.dylib", SDKDir + "lib/libembree.2.14.0.dylib");
+            RuntimeDependencies.Add("$(TargetOutputDir)/libtbb.dylib", SDKDir + "lib/libtbb.dylib");
+            RuntimeDependencies.Add("$(TargetOutputDir)/libtbbmalloc.dylib", SDKDir + "lib/libtbbmalloc.dylib");
 
             PrivateDefinitions.Add("USE_EMBREE=1");
         }
