Merge fix for undefined macro warning in Resonance.

diff --git a/Engine/Plugins/Runtime/ResonanceAudio/Source/ResonanceAudio/Private/ResonanceAudioLibrary/resonance_audio/utils/vorbis_stream_encoder.h b/Engine/Plugins/Runtime/ResonanceAudio/Source/ResonanceAudio/Private/ResonanceAudioLibrary/resonance_audio/utils/vorbis_stream_encoder.h
--- a/Engine/Plugins/Runtime/ResonanceAudio/Source/ResonanceAudio/Private/ResonanceAudioLibrary/resonance_audio/utils/vorbis_stream_encoder.h
+++ b/Engine/Plugins/Runtime/ResonanceAudio/Source/ResonanceAudio/Private/ResonanceAudioLibrary/resonance_audio/utils/vorbis_stream_encoder.h
@@ -35,6 +35,7 @@ _Pragma("clang diagnostic ignored \"-Wshadow\"")
 #pragma warning(disable:6386) /* Buffer overrun while writing to '<array>': the writable size is '<array size>' bytes, but '<write size>' bytes might be written */
 #pragma warning(disable:4244) /* 'initializing': conversion from '<type>' to '<type>', possible loss of data */
 #pragma warning(disable:4121) /* '<type>' alignment of a member was sensitive to packing */
+#pragma warning(disable:4668) /* '<macro>' is not defined as a preprocessor macro, replacing with '0' for '#if/#elif' */
 #endif
 
 #include <fstream>
