Potential fix for linker warning. hide visibility of the eigen includes in resonance to ld.

diff --git a/Engine/Plugins/Runtime/ResonanceAudio/Source/ResonanceAudio/Private/ResonanceEigenHelper.h b/Engine/Plugins/Runtime/ResonanceAudio/Source/ResonanceAudio/Private/ResonanceEigenHelper.h
--- a/Engine/Plugins/Runtime/ResonanceAudio/Source/ResonanceAudio/Private/ResonanceEigenHelper.h
+++ b/Engine/Plugins/Runtime/ResonanceAudio/Source/ResonanceAudio/Private/ResonanceEigenHelper.h
@@ -6,6 +6,7 @@
 #if defined(__clang__)
 _Pragma("clang diagnostic push") \
 _Pragma("clang diagnostic ignored \"-Wshadow\"")
+#pragma GCC visibility push(hidden)
 #elif defined(_MSC_VER)
 #pragma warning(push)
 #pragma warning(disable:6294) /* Ill-defined for-loop:  initial condition does not satisfy test.  Loop body not executed. */
@@ -23,6 +24,7 @@ _Pragma("clang diagnostic ignored \"-Wshadow\"")
 #include "Eigen/Dense"
 
 #if defined(__clang__)
+#pragma GCC visibility pop
 _Pragma("clang diagnostic pop")
 #elif defined(_MSC_VER)
 #pragma warning(pop)
