Flip back off the resonance reverb ignore hack now that we're confident we've fixed the issue. Resonance submix assets will now be loadable again.

diff --git a/Engine/Plugins/Runtime/ResonanceAudio/Source/ResonanceAudio/Private/ResonanceAudioReverb.cpp b/Engine/Plugins/Runtime/ResonanceAudio/Source/ResonanceAudio/Private/ResonanceAudioReverb.cpp
--- a/Engine/Plugins/Runtime/ResonanceAudio/Source/ResonanceAudio/Private/ResonanceAudioReverb.cpp
+++ b/Engine/Plugins/Runtime/ResonanceAudio/Source/ResonanceAudio/Private/ResonanceAudioReverb.cpp
@@ -9,7 +9,7 @@
 #include "ResonanceAudioSettings.h"
 #include "Sound/SoundSubmix.h"
 
-static int32 IgnoreUserResonanceSubmixCVar = 1;
+static int32 IgnoreUserResonanceSubmixCVar = 0;
 FAutoConsoleVariableRef CVarIgnoreUserResonanceSubmix(
 	TEXT("au.IgnoreUserResonanceSubmix"),
 	IgnoreUserResonanceSubmixCVar,
