Remove embree dependency for mac for google resonance, like we do for other non-windows platforms. since WITH_EMBREE is also used by MeshUtilities, we use a second macro to disable embree support specifically for resonance.

diff --git a/Engine/Plugins/Runtime/ResonanceAudio/Source/ResonanceAudio/Private/NullEmbreeImplementation.cpp b/Engine/Plugins/Runtime/ResonanceAudio/Source/ResonanceAudio/Private/NullEmbreeImplementation.cpp
--- a/Engine/Plugins/Runtime/ResonanceAudio/Source/ResonanceAudio/Private/NullEmbreeImplementation.cpp
+++ b/Engine/Plugins/Runtime/ResonanceAudio/Source/ResonanceAudio/Private/NullEmbreeImplementation.cpp
@@ -3,7 +3,11 @@
 // In order to support resonance cross platform, we need to be able to compile Resonance on platforms we may not support Embree on.
 // For those, we noop on all RTC calls:
 
-#if !USE_EMBREE
+#if !defined(USE_EMBREE)
+#define USE_EMBREE 0
+#endif 
+
+#if !USE_EMBREE || !USE_EMBRE_FOR_RESONANCE
 
 #include "ResonanceEmbreeHelper.h"
 #include "CoreMinimal.h"
@@ -84,4 +88,4 @@ void rtcSetUserData(RTCScene scene, unsigned geomID, void* ptr)
 	checkf(false, TEXT("Embree not compiled for this platform! To support raytracing in Resonance, compile Embree for this platform and add it in ResonanceAudio.build.cs."));
 }
 
-#endif // USE_EMBREE
\ No newline at end of file
+#endif // !USE_EMBREE || !USE_EMBRE_FOR_RESONANCE
\ No newline at end of file
diff --git a/Engine/Plugins/Runtime/ResonanceAudio/Source/ResonanceAudio/ResonanceAudio.Build.cs b/Engine/Plugins/Runtime/ResonanceAudio/Source/ResonanceAudio/ResonanceAudio.Build.cs
--- a/Engine/Plugins/Runtime/ResonanceAudio/Source/ResonanceAudio/ResonanceAudio.Build.cs
+++ b/Engine/Plugins/Runtime/ResonanceAudio/Source/ResonanceAudio/ResonanceAudio.Build.cs
@@ -92,7 +92,7 @@ public class ResonanceAudio : ModuleRules
         // Always use the official version of IntelTBB
         string IntelTBBLibs = Target.UEThirdPartySourceDirectory + "IntelTBB/IntelTBB-2019u8/lib/";
 
-        //Embree support:
+		//Embree support:
         if (Target.Platform == UnrealTargetPlatform.Win64)
         {
             string SDKDir = Target.UEThirdPartySourceDirectory + "IntelEmbree/Embree2140/Win64/";
@@ -103,25 +103,24 @@ public class ResonanceAudio : ModuleRules
             RuntimeDependencies.Add("$(TargetOutputDir)/tbb.dll", IntelTBBLibs + "Win64/vc14/tbb.dll");
             RuntimeDependencies.Add("$(TargetOutputDir)/tbbmalloc.dll", IntelTBBLibs + "Win64/vc14/tbbmalloc.dll");
 			PublicDefinitions.Add("USE_EMBREE=1");
-        }
+			PrivateDefinitions.Add("USE_EMBRE_FOR_RESONANCE=1");
+		}
         else if (Target.Platform == UnrealTargetPlatform.Mac)
         {
-            string SDKDir = Target.UEThirdPartySourceDirectory + "IntelEmbree/Embree2140/MacOSX/";
-
-            PublicIncludePaths.Add(SDKDir + "include");
-            PublicAdditionalLibraries.Add(SDKDir + "lib/libembree.2.14.0.dylib");
-            RuntimeDependencies.Add("$(TargetOutputDir)/libembree.2.14.0.dylib", SDKDir + "lib/libembree.2.14.0.dylib");
-            RuntimeDependencies.Add("$(TargetOutputDir)/libtbb.dylib", IntelTBBLibs + "Mac/libtbb.dylib");
-            RuntimeDependencies.Add("$(TargetOutputDir)/libtbbmalloc.dylib", IntelTBBLibs + "Mac/libtbbmalloc.dylib");
-            PrivateDefinitions.Add("USE_EMBREE=1");
-        }
-        else
+			// In platforms that don't support Embree, we implement no-op versions of the functions.
+			string SDKDir = Target.UEThirdPartySourceDirectory + "IntelEmbree/Embree2140/";
+			PublicIncludePaths.Add(SDKDir + "include");
+			PrivateDefinitions.Add("USE_EMBRE_FOR_RESONANCE=0");
+			PrivateDefinitions.Add("EMBREE_STATIC_LIB=1");
+		}
+		else
         {
             // In platforms that don't support Embree, we implement no-op versions of the functions.
             string SDKDir = Target.UEThirdPartySourceDirectory + "IntelEmbree/Embree2140/";
             PublicIncludePaths.Add(SDKDir + "include");
-            PublicDefinitions.Add("USE_EMBREE=0");
-            PrivateDefinitions.Add("EMBREE_STATIC_LIB=1");
+			PrivateDefinitions.Add("USE_EMBRE_FOR_RESONANCE=0");
+			PrivateDefinitions.Add("USE_EMBREE=0");
+			PrivateDefinitions.Add("EMBREE_STATIC_LIB=1");
         }
     }
 }
