Fixing PVS 7.7 Issues:

warning V1028: Possible overflow. Consider casting operands, not the result.

diff --git a/Engine/Plugins/Runtime/ResonanceAudio/Source/ResonanceAudio/Private/ResonanceAudioLibrary/resonance_audio/ambisonics/associated_legendre_polynomials_generator.cc b/Engine/Plugins/Runtime/ResonanceAudio/Source/ResonanceAudio/Private/ResonanceAudioLibrary/resonance_audio/ambisonics/associated_legendre_polynomials_generator.cc
--- a/Engine/Plugins/Runtime/ResonanceAudio/Source/ResonanceAudio/Private/ResonanceAudioLibrary/resonance_audio/ambisonics/associated_legendre_polynomials_generator.cc
+++ b/Engine/Plugins/Runtime/ResonanceAudio/Source/ResonanceAudio/Private/ResonanceAudioLibrary/resonance_audio/ambisonics/associated_legendre_polynomials_generator.cc
@@ -95,10 +95,12 @@ size_t AssociatedLegendrePolynomialsGenerator::GetIndex(int degree,
                                                         int order) const {
   CheckIndexValidity(degree, order);
   size_t result;
+  const size_t sdegree = static_cast<size_t>(degree);
+  const size_t sorder = static_cast<size_t>(order);
   if (compute_negative_order_) {
-    result = static_cast<size_t>(degree * (degree + 1) + order);
+    result = sdegree * (sdegree + 1) + sorder;
   } else {
-    result = static_cast<size_t>((degree * (degree + 1)) / 2 + order);
+    result = (sdegree * (sdegree + 1)) / 2 + sorder;
   }
   DCHECK_GE(result, 0U);
   DCHECK_LT(result, GetNumValues());
diff --git a/Engine/Plugins/Runtime/ResonanceAudio/Source/ResonanceAudio/Private/ResonanceAudioLibrary/resonance_audio/ambisonics/utils.h b/Engine/Plugins/Runtime/ResonanceAudio/Source/ResonanceAudio/Private/ResonanceAudioLibrary/resonance_audio/ambisonics/utils.h
--- a/Engine/Plugins/Runtime/ResonanceAudio/Source/ResonanceAudio/Private/ResonanceAudioLibrary/resonance_audio/ambisonics/utils.h
+++ b/Engine/Plugins/Runtime/ResonanceAudio/Source/ResonanceAudio/Private/ResonanceAudioLibrary/resonance_audio/ambisonics/utils.h
@@ -57,7 +57,8 @@ struct GetNumPeriphonicComponentsStatic {
 // Returns the number of spherical harmonics for a periphonic ambisonic sound
 // field of |ambisonic_order|.
 inline size_t GetNumPeriphonicComponents(int ambisonic_order) {
-  return static_cast<size_t>((ambisonic_order + 1) * (ambisonic_order + 1));
+  const size_t sambisonic_order = static_cast<size_t>(ambisonic_order);
+  return (sambisonic_order + 1) * (sambisonic_order + 1);
 }
 
 // Returns the number of periphonic spherical harmonics (SHs) for a particular
