c++17 fixes

diff --git a/Engine/Plugins/NetcodeUnitTest/NetcodeUnitTest/Source/NetcodeUnitTest/Private/UnitTests/VMReflection.cpp b/Engine/Plugins/NetcodeUnitTest/NetcodeUnitTest/Source/NetcodeUnitTest/Private/UnitTests/VMReflection.cpp
--- a/Engine/Plugins/NetcodeUnitTest/NetcodeUnitTest/Source/NetcodeUnitTest/Private/UnitTests/VMReflection.cpp
+++ b/Engine/Plugins/NetcodeUnitTest/NetcodeUnitTest/Source/NetcodeUnitTest/Private/UnitTests/VMReflection.cpp
@@ -47,7 +47,7 @@ bool UVMReflection::ExecuteUnitTest()
 		(void)(UObject*)(FVMReflection(TestObjA)->*"AObjectRef", &bOriginalError);
 
 		bool bError = false;
-		(FString)(FVMReflection(TestObjA)->*"AObjectRef", &bError);
+		(void)(FString)(FVMReflection(TestObjA)->*"AObjectRef", &bError);
 
 
 		TestResults.Add(TEXT("Reflection casting error"), (!bOriginalError && bError));
diff --git a/Engine/Plugins/Runtime/ResonanceAudio/Source/ResonanceAudio/Private/ResonanceAudioLibrary/resonance_audio/base/aligned_allocator.h b/Engine/Plugins/Runtime/ResonanceAudio/Source/ResonanceAudio/Private/ResonanceAudioLibrary/resonance_audio/base/aligned_allocator.h
--- a/Engine/Plugins/Runtime/ResonanceAudio/Source/ResonanceAudio/Private/ResonanceAudioLibrary/resonance_audio/base/aligned_allocator.h
+++ b/Engine/Plugins/Runtime/ResonanceAudio/Source/ResonanceAudio/Private/ResonanceAudioLibrary/resonance_audio/base/aligned_allocator.h
@@ -84,9 +84,9 @@ class AlignedAllocator : public std::allocator<Type> {
  public:
 
 #if __cplusplus >= 201703L
-  using Pointer = pointer;
-  using ConstPointer = const_pointer;
-  using SizeType = size_type;
+  using Pointer = typename std::allocator_traits<std::allocator<Type>>::pointer;
+  using ConstPointer = typename std::allocator_traits<std::allocator<Type>>::const_pointer;
+  using SizeType = typename std::allocator_traits<std::allocator<Type>>::size_type;
 
 #else
   typedef typename std::allocator<Type>::pointer Pointer;
