Properly include Eigen as a Third Party Include to avoid deprecation warnings

diff --git a/Engine/Plugins/Runtime/ResonanceAudio/Source/ResonanceAudio/Private/ResonanceEigenHelper.h b/Engine/Plugins/Runtime/ResonanceAudio/Source/ResonanceAudio/Private/ResonanceEigenHelper.h
--- a/Engine/Plugins/Runtime/ResonanceAudio/Source/ResonanceAudio/Private/ResonanceEigenHelper.h
+++ b/Engine/Plugins/Runtime/ResonanceAudio/Source/ResonanceAudio/Private/ResonanceEigenHelper.h
@@ -3,6 +3,8 @@
 #pragma once
 
 // This file provides a way to include Eigen/Core without producing static analysis warnings.
+THIRD_PARTY_INCLUDES_START
+
 #if defined(__clang__)
 _Pragma("clang diagnostic push") \
 _Pragma("clang diagnostic ignored \"-Wshadow\"")
@@ -28,3 +30,5 @@ _Pragma("clang diagnostic pop")
 #pragma warning(pop)
 #endif
 
+THIRD_PARTY_INCLUDES_END
+
