Limit disabling to deprecation instead of full third party includes because the full one causes compile errors in this context

diff --git a/Engine/Plugins/Runtime/ResonanceAudio/Source/ResonanceAudio/Private/ResonanceEigenHelper.h b/Engine/Plugins/Runtime/ResonanceAudio/Source/ResonanceAudio/Private/ResonanceEigenHelper.h
--- a/Engine/Plugins/Runtime/ResonanceAudio/Source/ResonanceAudio/Private/ResonanceEigenHelper.h
+++ b/Engine/Plugins/Runtime/ResonanceAudio/Source/ResonanceAudio/Private/ResonanceEigenHelper.h
@@ -3,7 +3,8 @@
 #pragma once
 
 // This file provides a way to include Eigen/Core without producing static analysis warnings.
-THIRD_PARTY_INCLUDES_START
+
+#include "GenericPlatform/GenericPlatformCompilerPreSetup.h"
 
 #if defined(__clang__)
 _Pragma("clang diagnostic push") \
@@ -20,15 +21,15 @@ _Pragma("clang diagnostic ignored \"-Wshadow\"")
 #pragma warning(disable:4702) /* unreachable code */
 #endif
 
+PRAGMA_DISABLE_DEPRECATION_WARNINGS
 
 #include "Eigen/Core"
 #include "Eigen/Dense"
 
+PRAGMA_ENABLE_DEPRECATION_WARNINGS
+
 #if defined(__clang__)
 _Pragma("clang diagnostic pop")
 #elif defined(_MSC_VER)
 #pragma warning(pop)
 #endif
-
-THIRD_PARTY_INCLUDES_END
-
