Include Platform.h instead of PreCompilerSetup.h directly

diff --git a/Engine/Plugins/Runtime/ResonanceAudio/Source/ResonanceAudio/Private/ResonanceEigenHelper.h b/Engine/Plugins/Runtime/ResonanceAudio/Source/ResonanceAudio/Private/ResonanceEigenHelper.h
--- a/Engine/Plugins/Runtime/ResonanceAudio/Source/ResonanceAudio/Private/ResonanceEigenHelper.h
+++ b/Engine/Plugins/Runtime/ResonanceAudio/Source/ResonanceAudio/Private/ResonanceEigenHelper.h
@@ -4,7 +4,7 @@
 
 // This file provides a way to include Eigen/Core without producing static analysis warnings.
 
-#include "GenericPlatform/GenericPlatformCompilerPreSetup.h"
+#include "HAL/Platform.h"
 
 #if defined(__clang__)
 _Pragma("clang diagnostic push") \
