Fix for UE-111191- //UE5/Main - NonUnity Compile UnrealEditor Win64 - '__STDC_WANT_SECURE_LIB__' is not defined as a preprocessor macro

diff --git a/Engine/Plugins/Runtime/ResonanceAudio/Source/ResonanceAudio/Private/ResonanceAudioLibrary/resonance_audio/utils/wav_reader.h b/Engine/Plugins/Runtime/ResonanceAudio/Source/ResonanceAudio/Private/ResonanceAudioLibrary/resonance_audio/utils/wav_reader.h
--- a/Engine/Plugins/Runtime/ResonanceAudio/Source/ResonanceAudio/Private/ResonanceAudioLibrary/resonance_audio/utils/wav_reader.h
+++ b/Engine/Plugins/Runtime/ResonanceAudio/Source/ResonanceAudio/Private/ResonanceAudioLibrary/resonance_audio/utils/wav_reader.h
@@ -17,6 +17,8 @@ limitations under the License.
 #ifndef RESONANCE_AUDIO_UTILS_WAV_READER_H_
 #define RESONANCE_AUDIO_UTILS_WAV_READER_H_
 
+#include "HAL/PlatformCrt.h"
+
 #include <cstdint>
 #include <istream>
 
