Enable Google Resonance Audio as Default Engine Plugin

diff --git a/Engine/Plugins/Runtime/ResonanceAudio/ResonanceAudio.uplugin b/Engine/Plugins/Runtime/ResonanceAudio/ResonanceAudio.uplugin
--- a/Engine/Plugins/Runtime/ResonanceAudio/ResonanceAudio.uplugin
+++ b/Engine/Plugins/Runtime/ResonanceAudio/ResonanceAudio.uplugin
@@ -10,7 +10,7 @@
   "DocsURL": "https://developers.google.com/resonance-audio/develop/unreal/getting-started",
   "MarketplaceURL": "",
   "SupportURL": "https://developers.google.com/resonance-audio/community/connect",
-  "EnabledByDefault": false,
+  "EnabledByDefault": true,
   "CanContainContent": true,
   "IsBetaVersion": true,
   "Installed": false,
