Manual fix for use of raw pointer member instead of TObjectPtr.

diff --git a/Engine/Plugins/Runtime/ResonanceAudio/Source/ResonanceAudio/Private/ResonanceAudioDirectivityVisualizer.h b/Engine/Plugins/Runtime/ResonanceAudio/Source/ResonanceAudio/Private/ResonanceAudioDirectivityVisualizer.h
--- a/Engine/Plugins/Runtime/ResonanceAudio/Source/ResonanceAudio/Private/ResonanceAudioDirectivityVisualizer.h
+++ b/Engine/Plugins/Runtime/ResonanceAudio/Source/ResonanceAudio/Private/ResonanceAudioDirectivityVisualizer.h
@@ -47,8 +47,8 @@ private:
 #endif
 
 	UPROPERTY()
-	UMaterial* Material;
+	TObjectPtr<UMaterial> Material;
 
 	UPROPERTY()
-	UResonanceAudioSpatializationSourceSettings* Settings;
+	TObjectPtr<UResonanceAudioSpatializationSourceSettings> Settings;
 };
