Fix up visibility warning for Eigen symbols on clang-based compilers (macOS/iOS/tvOS/Linux).

#rnx

diff --git a/Engine/Plugins/Runtime/ResonanceAudio/Source/ResonanceAudio/Private/ResonanceEigenHelper.h b/Engine/Plugins/Runtime/ResonanceAudio/Source/ResonanceAudio/Private/ResonanceEigenHelper.h
--- a/Engine/Plugins/Runtime/ResonanceAudio/Source/ResonanceAudio/Private/ResonanceEigenHelper.h
+++ b/Engine/Plugins/Runtime/ResonanceAudio/Source/ResonanceAudio/Private/ResonanceEigenHelper.h
@@ -21,12 +21,14 @@ _Pragma("clang diagnostic ignored \"-Wshadow\"")
 #pragma warning(disable:4702) /* unreachable code */
 #endif
 
+PRAGMA_DEFAULT_VISIBILITY_START
 PRAGMA_DISABLE_DEPRECATION_WARNINGS
 
 #include "Eigen/Core"
 #include "Eigen/Dense"
 
 PRAGMA_ENABLE_DEPRECATION_WARNINGS
+PRAGMA_DEFAULT_VISIBILITY_END
 
 #if defined(__clang__)
 _Pragma("clang diagnostic pop")
