Removed libc seed being hardcoded to 1 in Resonance audio plugin, overriding libc seed previously set during UE initialization

#rnx

diff --git a/Engine/Plugins/Runtime/ResonanceAudio/Source/ResonanceAudio/Private/ResonanceAudioLibrary/resonance_audio/dsp/spectral_reverb.cc b/Engine/Plugins/Runtime/ResonanceAudio/Source/ResonanceAudio/Private/ResonanceAudioLibrary/resonance_audio/dsp/spectral_reverb.cc
--- a/Engine/Plugins/Runtime/ResonanceAudio/Source/ResonanceAudio/Private/ResonanceAudioLibrary/resonance_audio/dsp/spectral_reverb.cc
+++ b/Engine/Plugins/Runtime/ResonanceAudio/Source/ResonanceAudio/Private/ResonanceAudioLibrary/resonance_audio/dsp/spectral_reverb.cc
@@ -98,8 +98,6 @@ SpectralReverb::SpectralReverb(int sample_rate, size_t frames_per_buffer)
   DCHECK_GT(sample_rate, 0);
   DCHECK_GT(frames_per_buffer_, 0U);
 
-  // Seed std::rand, used for phase selection.
-  std::srand(1);
   GenerateRandomPhaseBuffer();
   GenerateAnalysisWindow();
   InitializeCircularBuffersAndAccumulators();
