Remove embree dependency from ResonanceAudio
- Deleted geometrical_acoustics folder.

diff --git a/Engine/Plugins/Runtime/ResonanceAudio/Source/ResonanceAudio/Private/NullEmbreeImplementation.cpp b/Engine/Plugins/Runtime/ResonanceAudio/Source/ResonanceAudio/Private/NullEmbreeImplementation.cpp
deleted file mode 100644
diff --git a/Engine/Plugins/Runtime/ResonanceAudio/Source/ResonanceAudio/Private/ResonanceAudioLibrary/resonance_audio/geometrical_acoustics/acoustic_listener.h b/Engine/Plugins/Runtime/ResonanceAudio/Source/ResonanceAudio/Private/ResonanceAudioLibrary/resonance_audio/geometrical_acoustics/acoustic_listener.h
deleted file mode 100644
diff --git a/Engine/Plugins/Runtime/ResonanceAudio/Source/ResonanceAudio/Private/ResonanceAudioLibrary/resonance_audio/geometrical_acoustics/acoustic_ray.cc b/Engine/Plugins/Runtime/ResonanceAudio/Source/ResonanceAudio/Private/ResonanceAudioLibrary/resonance_audio/geometrical_acoustics/acoustic_ray.cc
deleted file mode 100644
diff --git a/Engine/Plugins/Runtime/ResonanceAudio/Source/ResonanceAudio/Private/ResonanceAudioLibrary/resonance_audio/geometrical_acoustics/acoustic_ray.h b/Engine/Plugins/Runtime/ResonanceAudio/Source/ResonanceAudio/Private/ResonanceAudioLibrary/resonance_audio/geometrical_acoustics/acoustic_ray.h
deleted file mode 100644
diff --git a/Engine/Plugins/Runtime/ResonanceAudio/Source/ResonanceAudio/Private/ResonanceAudioLibrary/resonance_audio/geometrical_acoustics/acoustic_source.h b/Engine/Plugins/Runtime/ResonanceAudio/Source/ResonanceAudio/Private/ResonanceAudioLibrary/resonance_audio/geometrical_acoustics/acoustic_source.h
deleted file mode 100644
diff --git a/Engine/Plugins/Runtime/ResonanceAudio/Source/ResonanceAudio/Private/ResonanceAudioLibrary/resonance_audio/geometrical_acoustics/collection_kernel.cc b/Engine/Plugins/Runtime/ResonanceAudio/Source/ResonanceAudio/Private/ResonanceAudioLibrary/resonance_audio/geometrical_acoustics/collection_kernel.cc
deleted file mode 100644
diff --git a/Engine/Plugins/Runtime/ResonanceAudio/Source/ResonanceAudio/Private/ResonanceAudioLibrary/resonance_audio/geometrical_acoustics/collection_kernel.h b/Engine/Plugins/Runtime/ResonanceAudio/Source/ResonanceAudio/Private/ResonanceAudioLibrary/resonance_audio/geometrical_acoustics/collection_kernel.h
deleted file mode 100644
diff --git a/Engine/Plugins/Runtime/ResonanceAudio/Source/ResonanceAudio/Private/ResonanceAudioLibrary/resonance_audio/geometrical_acoustics/estimating_rt60.cc b/Engine/Plugins/Runtime/ResonanceAudio/Source/ResonanceAudio/Private/ResonanceAudioLibrary/resonance_audio/geometrical_acoustics/estimating_rt60.cc
deleted file mode 100644
diff --git a/Engine/Plugins/Runtime/ResonanceAudio/Source/ResonanceAudio/Private/ResonanceAudioLibrary/resonance_audio/geometrical_acoustics/estimating_rt60.h b/Engine/Plugins/Runtime/ResonanceAudio/Source/ResonanceAudio/Private/ResonanceAudioLibrary/resonance_audio/geometrical_acoustics/estimating_rt60.h
deleted file mode 100644
diff --git a/Engine/Plugins/Runtime/ResonanceAudio/Source/ResonanceAudio/Private/ResonanceAudioLibrary/resonance_audio/geometrical_acoustics/impulse_response_computer.cc b/Engine/Plugins/Runtime/ResonanceAudio/Source/ResonanceAudio/Private/ResonanceAudioLibrary/resonance_audio/geometrical_acoustics/impulse_response_computer.cc
deleted file mode 100644
diff --git a/Engine/Plugins/Runtime/ResonanceAudio/Source/ResonanceAudio/Private/ResonanceAudioLibrary/resonance_audio/geometrical_acoustics/impulse_response_computer.h b/Engine/Plugins/Runtime/ResonanceAudio/Source/ResonanceAudio/Private/ResonanceAudioLibrary/resonance_audio/geometrical_acoustics/impulse_response_computer.h
deleted file mode 100644
diff --git a/Engine/Plugins/Runtime/ResonanceAudio/Source/ResonanceAudio/Private/ResonanceAudioLibrary/resonance_audio/geometrical_acoustics/mesh.h b/Engine/Plugins/Runtime/ResonanceAudio/Source/ResonanceAudio/Private/ResonanceAudioLibrary/resonance_audio/geometrical_acoustics/mesh.h
deleted file mode 100644
diff --git a/Engine/Plugins/Runtime/ResonanceAudio/Source/ResonanceAudio/Private/ResonanceAudioLibrary/resonance_audio/geometrical_acoustics/parallel_for.cc b/Engine/Plugins/Runtime/ResonanceAudio/Source/ResonanceAudio/Private/ResonanceAudioLibrary/resonance_audio/geometrical_acoustics/parallel_for.cc
deleted file mode 100644
diff --git a/Engine/Plugins/Runtime/ResonanceAudio/Source/ResonanceAudio/Private/ResonanceAudioLibrary/resonance_audio/geometrical_acoustics/parallel_for.h b/Engine/Plugins/Runtime/ResonanceAudio/Source/ResonanceAudio/Private/ResonanceAudioLibrary/resonance_audio/geometrical_acoustics/parallel_for.h
deleted file mode 100644
diff --git a/Engine/Plugins/Runtime/ResonanceAudio/Source/ResonanceAudio/Private/ResonanceAudioLibrary/resonance_audio/geometrical_acoustics/path.h b/Engine/Plugins/Runtime/ResonanceAudio/Source/ResonanceAudio/Private/ResonanceAudioLibrary/resonance_audio/geometrical_acoustics/path.h
deleted file mode 100644
diff --git a/Engine/Plugins/Runtime/ResonanceAudio/Source/ResonanceAudio/Private/ResonanceAudioLibrary/resonance_audio/geometrical_acoustics/path_tracer.cc b/Engine/Plugins/Runtime/ResonanceAudio/Source/ResonanceAudio/Private/ResonanceAudioLibrary/resonance_audio/geometrical_acoustics/path_tracer.cc
deleted file mode 100644
diff --git a/Engine/Plugins/Runtime/ResonanceAudio/Source/ResonanceAudio/Private/ResonanceAudioLibrary/resonance_audio/geometrical_acoustics/path_tracer.h b/Engine/Plugins/Runtime/ResonanceAudio/Source/ResonanceAudio/Private/ResonanceAudioLibrary/resonance_audio/geometrical_acoustics/path_tracer.h
deleted file mode 100644
diff --git a/Engine/Plugins/Runtime/ResonanceAudio/Source/ResonanceAudio/Private/ResonanceAudioLibrary/resonance_audio/geometrical_acoustics/proxy_room_estimator.cc b/Engine/Plugins/Runtime/ResonanceAudio/Source/ResonanceAudio/Private/ResonanceAudioLibrary/resonance_audio/geometrical_acoustics/proxy_room_estimator.cc
deleted file mode 100644
diff --git a/Engine/Plugins/Runtime/ResonanceAudio/Source/ResonanceAudio/Private/ResonanceAudioLibrary/resonance_audio/geometrical_acoustics/proxy_room_estimator.h b/Engine/Plugins/Runtime/ResonanceAudio/Source/ResonanceAudio/Private/ResonanceAudioLibrary/resonance_audio/geometrical_acoustics/proxy_room_estimator.h
deleted file mode 100644
diff --git a/Engine/Plugins/Runtime/ResonanceAudio/Source/ResonanceAudio/Private/ResonanceAudioLibrary/resonance_audio/geometrical_acoustics/reflection_kernel.cc b/Engine/Plugins/Runtime/ResonanceAudio/Source/ResonanceAudio/Private/ResonanceAudioLibrary/resonance_audio/geometrical_acoustics/reflection_kernel.cc
deleted file mode 100644
diff --git a/Engine/Plugins/Runtime/ResonanceAudio/Source/ResonanceAudio/Private/ResonanceAudioLibrary/resonance_audio/geometrical_acoustics/reflection_kernel.h b/Engine/Plugins/Runtime/ResonanceAudio/Source/ResonanceAudio/Private/ResonanceAudioLibrary/resonance_audio/geometrical_acoustics/reflection_kernel.h
deleted file mode 100644
diff --git a/Engine/Plugins/Runtime/ResonanceAudio/Source/ResonanceAudio/Private/ResonanceAudioLibrary/resonance_audio/geometrical_acoustics/sampling.h b/Engine/Plugins/Runtime/ResonanceAudio/Source/ResonanceAudio/Private/ResonanceAudioLibrary/resonance_audio/geometrical_acoustics/sampling.h
deleted file mode 100644
diff --git a/Engine/Plugins/Runtime/ResonanceAudio/Source/ResonanceAudio/Private/ResonanceAudioLibrary/resonance_audio/geometrical_acoustics/scene_manager.cc b/Engine/Plugins/Runtime/ResonanceAudio/Source/ResonanceAudio/Private/ResonanceAudioLibrary/resonance_audio/geometrical_acoustics/scene_manager.cc
deleted file mode 100644
diff --git a/Engine/Plugins/Runtime/ResonanceAudio/Source/ResonanceAudio/Private/ResonanceAudioLibrary/resonance_audio/geometrical_acoustics/scene_manager.h b/Engine/Plugins/Runtime/ResonanceAudio/Source/ResonanceAudio/Private/ResonanceAudioLibrary/resonance_audio/geometrical_acoustics/scene_manager.h
deleted file mode 100644
diff --git a/Engine/Plugins/Runtime/ResonanceAudio/Source/ResonanceAudio/Private/ResonanceAudioLibrary/resonance_audio/geometrical_acoustics/sphere.cc b/Engine/Plugins/Runtime/ResonanceAudio/Source/ResonanceAudio/Private/ResonanceAudioLibrary/resonance_audio/geometrical_acoustics/sphere.cc
deleted file mode 100644
diff --git a/Engine/Plugins/Runtime/ResonanceAudio/Source/ResonanceAudio/Private/ResonanceAudioLibrary/resonance_audio/geometrical_acoustics/sphere.h b/Engine/Plugins/Runtime/ResonanceAudio/Source/ResonanceAudio/Private/ResonanceAudioLibrary/resonance_audio/geometrical_acoustics/sphere.h
deleted file mode 100644
diff --git a/Engine/Plugins/Runtime/ResonanceAudio/Source/ResonanceAudio/Private/ResonanceEmbreeHelper.h b/Engine/Plugins/Runtime/ResonanceAudio/Source/ResonanceAudio/Private/ResonanceEmbreeHelper.h
deleted file mode 100644
diff --git a/Engine/Plugins/Runtime/ResonanceAudio/Source/ResonanceAudio/ResonanceAudio.Build.cs b/Engine/Plugins/Runtime/ResonanceAudio/Source/ResonanceAudio/ResonanceAudio.Build.cs
--- a/Engine/Plugins/Runtime/ResonanceAudio/Source/ResonanceAudio/ResonanceAudio.Build.cs
+++ b/Engine/Plugins/Runtime/ResonanceAudio/Source/ResonanceAudio/ResonanceAudio.Build.cs
@@ -90,59 +90,5 @@ public class ResonanceAudio : ModuleRules
             PrivateDefinitions.Add("PFFFT_SIMD_DISABLE=1");
             PrivateDefinitions.Add("EIGEN_HAS_CXX11_MATH=0");
         }
-
-        // Always use the official version of IntelTBB
-        string IntelTBBLibs = Target.UEThirdPartyBinariesDirectory + "Intel/TBB/";
-        string IntelEmbreeLibs = Target.UEThirdPartyBinariesDirectory + "Intel/Embree/Embree2140";
-
-		//Embree support:
-        if (Target.Platform == UnrealTargetPlatform.Win64)
-        {
-            string SDKDir = Target.UEThirdPartySourceDirectory + "Intel/Embree/Embree2140/Win64/";
-
-            PublicIncludePaths.Add(SDKDir + "include");
-            PublicAdditionalLibraries.Add(SDKDir + "lib/embree.2.14.0.lib");
-            RuntimeDependencies.Add("$(TargetOutputDir)/embree.2.14.0.dll", SDKDir + "lib/embree.2.14.0.dll");
-            RuntimeDependencies.Add("$(TargetOutputDir)/tbb.dll", IntelTBBLibs + "Win64/tbb.dll");
-            RuntimeDependencies.Add("$(TargetOutputDir)/tbbmalloc.dll", IntelTBBLibs + "Win64/tbbmalloc.dll");
-			PublicDefinitions.Add("USE_EMBREE=1");
-			PrivateDefinitions.Add("USE_EMBRE_FOR_RESONANCE=1");
-		}
-        else if (Target.Platform == UnrealTargetPlatform.Mac)
-        {
-			// In platforms that don't support Embree, we implement no-op versions of the functions.
-			string SDKDir = Target.UEThirdPartySourceDirectory + "Intel/Embree/Embree2140/";
-			PublicIncludePaths.Add(SDKDir + "include");
-			PrivateDefinitions.Add("USE_EMBRE_FOR_RESONANCE=0");
-			PrivateDefinitions.Add("EMBREE_STATIC_LIB=1");
-		}
-		else if (Target.IsInPlatformGroup(UnrealPlatformGroup.Unix) && Target.Architecture.StartsWith("x86_64")) // no support for arm64 yet
-		{
-			string IncludeDir = Target.UEThirdPartySourceDirectory + "Intel/Embree/Embree2140/Linux/x86_64-unknown-linux-gnu";
-			string SDKDir = Path.Combine(IntelEmbreeLibs, "Linux/x86_64-unknown-linux-gnu/lib");
-
-			PublicIncludePaths.Add(Path.Combine(IncludeDir, "include"));
-			PublicAdditionalLibraries.Add(Path.Combine(SDKDir, "libembree.so"));
-			PublicAdditionalLibraries.Add(Path.Combine(IntelTBBLibs, "Linux/libtbb.so"));
-			// disabled for Linux atm due to a bug in libtbbmalloc on exit
-			// PublicAdditionalLibraries.Add(Path.Combine(IntelTBBLibs, "Linux/libtbbmalloc.so"));
-			RuntimeDependencies.Add(Path.Combine(SDKDir, "libembree.so"));
-			RuntimeDependencies.Add(Path.Combine(SDKDir, "libembree.so.2"));
-			RuntimeDependencies.Add(Path.Combine(IntelTBBLibs, "Linux/libtbb.so"));
-			RuntimeDependencies.Add(Path.Combine(IntelTBBLibs, "Linux/libtbb.so.2"));
-			// disabled for Linux atm due to a bug in libtbbmalloc on exit
-			// RuntimeDependencies.Add(Path.Combine(IntelTBBLibs, "Linux/libtbbmalloc.so"));
-			PublicDefinitions.Add("USE_EMBREE=1");
-			PrivateDefinitions.Add("USE_EMBRE_FOR_RESONANCE=1");
-		}
-		else
-        {
-            // In platforms that don't support Embree, we implement no-op versions of the functions.
-            string SDKDir = Target.UEThirdPartySourceDirectory + "Intel/Embree/Embree2140/";
-            PublicIncludePaths.Add(SDKDir + "include");
-			PrivateDefinitions.Add("USE_EMBRE_FOR_RESONANCE=0");
-			PrivateDefinitions.Add("USE_EMBREE=0");
-			PrivateDefinitions.Add("EMBREE_STATIC_LIB=1");
-        }
     }
 }
