Created a private PCH for ResonanceAudio.

Building module win64 development:
Before: 40.71 secs
After: 23.93 secs

Building ShooterGame win64 development:
Before: 208.61 seconds
After: 195.91 seconds

diff --git a/Engine/Plugins/Runtime/ResonanceAudio/Source/ResonanceAudio/Private/ResonanceAudioPrivatePCH.h b/Engine/Plugins/Runtime/ResonanceAudio/Source/ResonanceAudio/Private/ResonanceAudioPrivatePCH.h
new file mode 100644
--- /dev/null
+++ b/Engine/Plugins/Runtime/ResonanceAudio/Source/ResonanceAudio/Private/ResonanceAudioPrivatePCH.h
@@ -0,0 +1,8 @@
+// Copyright Epic Games, Inc. All Rights Reserved.
+
+#pragma once
+
+#include "ResonanceEigenHelper.h"
+#include "base/misc_math.h"
+#include "base/audio_buffer.h"
+#include "node/processing_node.h"
\ No newline at end of file
diff --git a/Engine/Plugins/Runtime/ResonanceAudio/Source/ResonanceAudio/ResonanceAudio.Build.cs b/Engine/Plugins/Runtime/ResonanceAudio/Source/ResonanceAudio/ResonanceAudio.Build.cs
--- a/Engine/Plugins/Runtime/ResonanceAudio/Source/ResonanceAudio/ResonanceAudio.Build.cs
+++ b/Engine/Plugins/Runtime/ResonanceAudio/Source/ResonanceAudio/ResonanceAudio.Build.cs
@@ -9,8 +9,10 @@ public class ResonanceAudio : ModuleRules
 	protected virtual bool bSupportsProceduralMesh { get { return true; } }
 
 	public ResonanceAudio(ReadOnlyTargetRules Target) : base(Target)
-    {
-        string ResonanceAudioPath = ModuleDirectory + "/Private/ResonanceAudioLibrary";
+	{
+		PrivatePCHHeaderFile = "Private/ResonanceAudioPrivatePCH.h";
+
+		string ResonanceAudioPath = ModuleDirectory + "/Private/ResonanceAudioLibrary";
         string ResonanceAudioLibraryPath = ModuleDirectory + "/Private/ResonanceAudioLibrary/resonance_audio";
         string PFFTPath = ModuleDirectory + "/Private/ResonanceAudioLibrary/third_party/pfft";
 
