Audio Plugin Telemetry Part 3
OSCModulationMixing, ResonanceAudio, and SoundCueTemplates

diff --git a/Engine/Plugins/Runtime/OSCModulationMixing/Source/OSCModulationMixing/Private/OSCModulationMixingStatics.cpp b/Engine/Plugins/Runtime/OSCModulationMixing/Source/OSCModulationMixing/Private/OSCModulationMixingStatics.cpp
--- a/Engine/Plugins/Runtime/OSCModulationMixing/Source/OSCModulationMixing/Private/OSCModulationMixingStatics.cpp
+++ b/Engine/Plugins/Runtime/OSCModulationMixing/Source/OSCModulationMixing/Private/OSCModulationMixingStatics.cpp
@@ -7,6 +7,7 @@
 #include "OSCModulationMixing.h"
 #include "OSCServer.h"
 
+#include "EngineAnalytics.h"
 #include "HAL/IConsoleManager.h"
 #include "SoundControlBus.h"
 #include "SoundControlBusMix.h"
@@ -75,6 +76,10 @@ FOSCAddress UOSCModulationMixingStatics::GetProfileLoadPath()
 
 FOSCAddress UOSCModulationMixingStatics::GetProfileSavePath()
 {
+	if (FEngineAnalytics::IsAvailable())
+	{
+		FEngineAnalytics::GetProvider().RecordEvent(TEXT("Audio.Usage.OSCModulationMixing.GetProfileSavePath"));
+	}
 	return OSCModulation::Addresses::ProfileSave;
 }
 
diff --git a/Engine/Plugins/Runtime/ResonanceAudio/Source/ResonanceAudioEditor/Private/ResonanceAudioReverbPluginPresetFactory.cpp b/Engine/Plugins/Runtime/ResonanceAudio/Source/ResonanceAudioEditor/Private/ResonanceAudioReverbPluginPresetFactory.cpp
--- a/Engine/Plugins/Runtime/ResonanceAudio/Source/ResonanceAudioEditor/Private/ResonanceAudioReverbPluginPresetFactory.cpp
+++ b/Engine/Plugins/Runtime/ResonanceAudio/Source/ResonanceAudioEditor/Private/ResonanceAudioReverbPluginPresetFactory.cpp
@@ -4,6 +4,8 @@
 
 #include "ResonanceAudioReverbPluginPresetFactory.h"
 #include "ResonanceAudioReverb.h"
+
+#include "EngineAnalytics.h"
 #include "UObject/ObjectMacros.h"
 #include "UObject/Object.h"
 
@@ -49,6 +51,10 @@ UResonanceAudioReverbPluginPresetFactory::UResonanceAudioReverbPluginPresetFacto
 
 UObject* UResonanceAudioReverbPluginPresetFactory::FactoryCreateNew(UClass* Class, UObject* InParent, FName InName, EObjectFlags Flags, UObject* Context, FFeedbackContext* Warn)
 {
+	if (FEngineAnalytics::IsAvailable())
+	{
+		FEngineAnalytics::GetProvider().RecordEvent(TEXT("Audio.Usage.ResonanceAudio.PresetCreated"));
+	}
 	return NewObject<UResonanceAudioReverbPluginPreset>(InParent, InName, Flags);
 }
 
diff --git a/Engine/Plugins/Runtime/ResonanceAudio/Source/ResonanceAudioEditor/Private/ResonanceAudioSpatializationSourceSettingsFactory.cpp b/Engine/Plugins/Runtime/ResonanceAudio/Source/ResonanceAudioEditor/Private/ResonanceAudioSpatializationSourceSettingsFactory.cpp
--- a/Engine/Plugins/Runtime/ResonanceAudio/Source/ResonanceAudioEditor/Private/ResonanceAudioSpatializationSourceSettingsFactory.cpp
+++ b/Engine/Plugins/Runtime/ResonanceAudio/Source/ResonanceAudioEditor/Private/ResonanceAudioSpatializationSourceSettingsFactory.cpp
@@ -4,6 +4,8 @@
 
 #include "ResonanceAudioSpatializationSourceSettingsFactory.h"
 #include "ResonanceAudioSpatializationSourceSettings.h"
+
+#include "EngineAnalytics.h"
 #include "UObject/ObjectMacros.h"
 #include "UObject/Object.h"
 
@@ -49,6 +51,10 @@ UResonanceAudioSpatializationSourceSettingsFactory::UResonanceAudioSpatializatio
 
 UObject* UResonanceAudioSpatializationSourceSettingsFactory::FactoryCreateNew(UClass* Class, UObject* InParent, FName InName, EObjectFlags Flags, UObject* Context, FFeedbackContext* Warn)
 {
+	if (FEngineAnalytics::IsAvailable())
+	{
+		FEngineAnalytics::GetProvider().RecordEvent(TEXT("Audio.Usage.ResonanceAudio.SpatializationSourceSettingsCreated"));
+	}
 	return Cast<UObject>(NewObject<UResonanceAudioSpatializationSourceSettings>(InParent, InName, Flags));
 }
 
diff --git a/Engine/Plugins/Runtime/SoundCueTemplates/Source/SoundCueTemplatesEditor/Private/SoundCueTemplateFactory.cpp b/Engine/Plugins/Runtime/SoundCueTemplates/Source/SoundCueTemplatesEditor/Private/SoundCueTemplateFactory.cpp
--- a/Engine/Plugins/Runtime/SoundCueTemplates/Source/SoundCueTemplatesEditor/Private/SoundCueTemplateFactory.cpp
+++ b/Engine/Plugins/Runtime/SoundCueTemplates/Source/SoundCueTemplatesEditor/Private/SoundCueTemplateFactory.cpp
@@ -6,6 +6,7 @@
 #include "ClassViewerModule.h"
 #include "EdGraph/EdGraphPin.h"
 #include "Engine/Engine.h"
+#include "EngineAnalytics.h"
 #include "Kismet2/SClassPickerDialog.h"
 #include "SoundCueGraph/SoundCueGraphNode.h"
 #include "Sound/SoundCue.h"
@@ -89,6 +90,12 @@ UObject* USoundCueTemplateFactory::FactoryCreateNew(UClass* Class, UObject* InPa
 			}
 
 			NewSoundCueTemplate->RebuildGraph(*NewSoundCueTemplate);
+
+			if (FEngineAnalytics::IsAvailable())
+			{
+				FEngineAnalytics::GetProvider().RecordEvent(TEXT("Audio.Usage.SoundCueTemplates.TemplateCreated"));
+			}
+
 			return NewSoundCueTemplate;
 		}
 	}
