Update Telemetry String in ResonanceAudio

diff --git a/Engine/Plugins/Runtime/ResonanceAudio/Source/ResonanceAudioEditor/Private/ResonanceAudioReverbPluginPresetFactory.cpp b/Engine/Plugins/Runtime/ResonanceAudio/Source/ResonanceAudioEditor/Private/ResonanceAudioReverbPluginPresetFactory.cpp
--- a/Engine/Plugins/Runtime/ResonanceAudio/Source/ResonanceAudioEditor/Private/ResonanceAudioReverbPluginPresetFactory.cpp
+++ b/Engine/Plugins/Runtime/ResonanceAudio/Source/ResonanceAudioEditor/Private/ResonanceAudioReverbPluginPresetFactory.cpp
@@ -53,7 +53,7 @@ UObject* UResonanceAudioReverbPluginPresetFactory::FactoryCreateNew(UClass* Clas
 {
 	if (FEngineAnalytics::IsAvailable())
 	{
-		FEngineAnalytics::GetProvider().RecordEvent(TEXT("Audio.Usage.ResonanceAudio.PresetCreated"));
+		FEngineAnalytics::GetProvider().RecordEvent(TEXT("Audio.Usage.ResonanceAudio.ReverbPresetCreated"));
 	}
 	return NewObject<UResonanceAudioReverbPluginPreset>(InParent, InName, Flags);
 }
