CIS fix for ResonanceAudio wav_reader.h
#rnx

diff --git a/Engine/Plugins/Runtime/ResonanceAudio/Source/ResonanceAudio/Private/ResonanceAudioLibrary/resonance_audio/utils/wav_reader.h b/Engine/Plugins/Runtime/ResonanceAudio/Source/ResonanceAudio/Private/ResonanceAudioLibrary/resonance_audio/utils/wav_reader.h
--- a/Engine/Plugins/Runtime/ResonanceAudio/Source/ResonanceAudio/Private/ResonanceAudioLibrary/resonance_audio/utils/wav_reader.h
+++ b/Engine/Plugins/Runtime/ResonanceAudio/Source/ResonanceAudio/Private/ResonanceAudioLibrary/resonance_audio/utils/wav_reader.h
@@ -19,10 +19,8 @@ limitations under the License.
 
 #include "HAL/PlatformCrt.h"
 
-THIRD_PARTY_INCLUDES_START
 #include <cstdint>
 #include <istream>
-THIRD_PARTY_INCLUDES_END
 
 #include "base/integral_types.h"
 
