CIS non-unity compile fix in ResonanceAudio
#rnx

diff --git a/Engine/Plugins/Runtime/ResonanceAudio/Source/ResonanceAudio/Private/ResonanceAudioDirectivityVisualizer.cpp b/Engine/Plugins/Runtime/ResonanceAudio/Source/ResonanceAudio/Private/ResonanceAudioDirectivityVisualizer.cpp
--- a/Engine/Plugins/Runtime/ResonanceAudio/Source/ResonanceAudio/Private/ResonanceAudioDirectivityVisualizer.cpp
+++ b/Engine/Plugins/Runtime/ResonanceAudio/Source/ResonanceAudio/Private/ResonanceAudioDirectivityVisualizer.cpp
@@ -3,6 +3,7 @@
 //
 
 #include "ResonanceAudioDirectivityVisualizer.h"
+#include "ResonanceAudioConstants.h"
 
 #include "UObject/ConstructorHelpers.h"
 
diff --git a/Engine/Plugins/Runtime/ResonanceAudio/Source/ResonanceAudio/Private/ResonanceAudioSpatialization.cpp b/Engine/Plugins/Runtime/ResonanceAudio/Source/ResonanceAudio/Private/ResonanceAudioSpatialization.cpp
--- a/Engine/Plugins/Runtime/ResonanceAudio/Source/ResonanceAudio/Private/ResonanceAudioSpatialization.cpp
+++ b/Engine/Plugins/Runtime/ResonanceAudio/Source/ResonanceAudio/Private/ResonanceAudioSpatialization.cpp
@@ -1,5 +1,6 @@
 #include "ResonanceAudioSpatialization.h"
 
+#include "ResonanceAudioConstants.h"
 #include "ResonanceAudioModule.h"
 #include "ResonanceAudioSettings.h"
 
diff --git a/Engine/Plugins/Runtime/ResonanceAudio/Source/ResonanceAudioEditor/Private/ResonanceAudioSpatializationSourceSettingsFactory.cpp b/Engine/Plugins/Runtime/ResonanceAudio/Source/ResonanceAudioEditor/Private/ResonanceAudioSpatializationSourceSettingsFactory.cpp
--- a/Engine/Plugins/Runtime/ResonanceAudio/Source/ResonanceAudioEditor/Private/ResonanceAudioSpatializationSourceSettingsFactory.cpp
+++ b/Engine/Plugins/Runtime/ResonanceAudio/Source/ResonanceAudioEditor/Private/ResonanceAudioSpatializationSourceSettingsFactory.cpp
@@ -4,6 +4,7 @@
 
 #include "ResonanceAudioSpatializationSourceSettingsFactory.h"
 #include "ResonanceAudioSpatializationSourceSettings.h"
+#include "ResonanceAudioConstants.h"
 
 #include "AudioAnalytics.h"
 #include "UObject/ObjectMacros.h"
