Static analysis cleanup for Google Resonance.
#p4v-preflight-copy 21594590

diff --git a/Engine/Plugins/Runtime/ResonanceAudio/Source/ResonanceAudio/Private/ResonanceAudioLibrary/resonance_audio/dsp/multi_channel_iir.cc b/Engine/Plugins/Runtime/ResonanceAudio/Source/ResonanceAudio/Private/ResonanceAudioLibrary/resonance_audio/dsp/multi_channel_iir.cc
--- a/Engine/Plugins/Runtime/ResonanceAudio/Source/ResonanceAudio/Private/ResonanceAudioLibrary/resonance_audio/dsp/multi_channel_iir.cc
+++ b/Engine/Plugins/Runtime/ResonanceAudio/Source/ResonanceAudio/Private/ResonanceAudioLibrary/resonance_audio/dsp/multi_channel_iir.cc
@@ -80,7 +80,7 @@ void MultiChannelIir::Process(AudioBuffer::Channel* interleaved_buffer) {
 
   const size_t num_channel_chunks = num_channels_ / SIMD_LENGTH;
   const size_t num_buffer_chunks = interleaved_buffer->size() / SIMD_LENGTH;
-  const size_t delay_length_in_chunks = num_channel_chunks * num_coefficients_;
+  const size_t delay_length_in_chunks = std::max(static_cast<size_t>(1), num_channel_chunks * num_coefficients_);
 
   for (size_t current_frame = 0, individual_frame = 0;
        current_frame < num_buffer_chunks;
diff --git a/Engine/Plugins/Runtime/ResonanceAudio/Source/ResonanceAudio/Private/ResonanceAudioLibrary/resonance_audio/dsp/partitioned_fft_filter.cc b/Engine/Plugins/Runtime/ResonanceAudio/Source/ResonanceAudio/Private/ResonanceAudioLibrary/resonance_audio/dsp/partitioned_fft_filter.cc
--- a/Engine/Plugins/Runtime/ResonanceAudio/Source/ResonanceAudio/Private/ResonanceAudioLibrary/resonance_audio/dsp/partitioned_fft_filter.cc
+++ b/Engine/Plugins/Runtime/ResonanceAudio/Source/ResonanceAudio/Private/ResonanceAudioLibrary/resonance_audio/dsp/partitioned_fft_filter.cc
@@ -95,7 +95,7 @@ void PartitionedFftFilter::ResetFreqDomainBuffers(size_t new_filter_size) {
   DCHECK_LE(filter_size_, max_filter_size_);
 
   const size_t old_num_partitions = num_partitions_;
-  num_partitions_ = filter_size_ / frames_per_buffer_;
+  num_partitions_ = std::max(static_cast<size_t>(1), filter_size_ / frames_per_buffer_);
   const size_t min_num_partitions =
       std::min(old_num_partitions, num_partitions_);
 
diff --git a/Engine/Plugins/Runtime/ResonanceAudio/Source/ResonanceAudio/Private/ResonanceAudioLibrary/resonance_audio/utils/task_thread_pool.cc b/Engine/Plugins/Runtime/ResonanceAudio/Source/ResonanceAudio/Private/ResonanceAudioLibrary/resonance_audio/utils/task_thread_pool.cc
--- a/Engine/Plugins/Runtime/ResonanceAudio/Source/ResonanceAudio/Private/ResonanceAudioLibrary/resonance_audio/utils/task_thread_pool.cc
+++ b/Engine/Plugins/Runtime/ResonanceAudio/Source/ResonanceAudio/Private/ResonanceAudioLibrary/resonance_audio/utils/task_thread_pool.cc
@@ -189,7 +189,10 @@ bool TaskThreadPool::RunOnWorkerThread(TaskThreadPool::TaskClosure closure) {
     std::lock_guard<std::mutex> lock(worker_available_mutex_);
     --num_worker_threads_available_;
   }
-  available_worker_thread->Run(std::move(closure));
+
+  if(available_worker_thread)
+    available_worker_thread->Run(std::move(closure));
+
   return true;
 }
 
