Static analysis fix in Google Resonance plugin (potential div by zero)

diff --git a/Engine/Plugins/Runtime/ResonanceAudio/Source/ResonanceAudio/Private/ResonanceAudioLibrary/resonance_audio/dsp/partitioned_fft_filter.cc b/Engine/Plugins/Runtime/ResonanceAudio/Source/ResonanceAudio/Private/ResonanceAudioLibrary/resonance_audio/dsp/partitioned_fft_filter.cc
--- a/Engine/Plugins/Runtime/ResonanceAudio/Source/ResonanceAudio/Private/ResonanceAudioLibrary/resonance_audio/dsp/partitioned_fft_filter.cc
+++ b/Engine/Plugins/Runtime/ResonanceAudio/Source/ResonanceAudio/Private/ResonanceAudioLibrary/resonance_audio/dsp/partitioned_fft_filter.cc
@@ -33,6 +33,8 @@ limitations under the License.
 
 namespace vraudio {
 
+static size_t clamp_to_non_zero(const size_t& val) { return (val < 1)? 1 : val; }
+
 PartitionedFftFilter::PartitionedFftFilter(size_t filter_size,
                                            size_t frames_per_buffer,
                                            FftManager* fft_manager)
@@ -225,7 +227,7 @@ void PartitionedFftFilter::Filter(const FreqDomainBuffer::Channel& input) {
   }
   // Our modulo based index.
   curr_front_buffer_ =
-      (curr_front_buffer_ + num_partitions_ - 1) % num_partitions_;
+      (curr_front_buffer_ + num_partitions_ - 1) % clamp_to_non_zero(num_partitions_);
   // Perform inverse FFT transform of |freq_domain_buffer_| and store the
   // result back in |filtered_time_domain_buffers_|.
   fft_manager_->TimeFromFreqDomain(
