Fixed "Audio" LLM child tags for Audio_Analysis, Audio_MetaSound and Audio_SpatializationPlugins.
* AudioAnalysis --> Audio_Analysis
* MetaSound --> Audio_MetaSound
* AudioSpatializationPlugins --> Audio_SpatializationPlugins

The LLM_DEFINE_TAG macro needs to use the "underscore convention" to name a child tag prefixed with the name of the parent tag.

diff --git a/Engine/Plugins/Runtime/AudioSynesthesia/Source/AudioSynesthesia/Private/AudioSynesthesiaModule.cpp b/Engine/Plugins/Runtime/AudioSynesthesia/Source/AudioSynesthesia/Private/AudioSynesthesiaModule.cpp
--- a/Engine/Plugins/Runtime/AudioSynesthesia/Source/AudioSynesthesia/Private/AudioSynesthesiaModule.cpp
+++ b/Engine/Plugins/Runtime/AudioSynesthesia/Source/AudioSynesthesia/Private/AudioSynesthesiaModule.cpp
@@ -18,7 +18,7 @@ namespace Audio
 {
 	void FAudioSynesthesiaModule::StartupModule()
 	{
-		LLM_SCOPE_BYTAG(AudioAnalysis);
+		AUDIO_ANALYSIS_LLM_SCOPE
 		FModuleManager::Get().LoadModuleChecked<IModuleInterface>("AudioSynesthesiaCore");
 		FModuleManager::Get().LoadModuleChecked<FAudioAnalyzerModule>("AudioAnalyzer");
 	}
diff --git a/Engine/Plugins/Runtime/AudioSynesthesia/Source/AudioSynesthesiaEditor/Private/AudioSynesthesiaEditorModule.cpp b/Engine/Plugins/Runtime/AudioSynesthesia/Source/AudioSynesthesiaEditor/Private/AudioSynesthesiaEditorModule.cpp
--- a/Engine/Plugins/Runtime/AudioSynesthesia/Source/AudioSynesthesiaEditor/Private/AudioSynesthesiaEditorModule.cpp
+++ b/Engine/Plugins/Runtime/AudioSynesthesia/Source/AudioSynesthesiaEditor/Private/AudioSynesthesiaEditorModule.cpp
@@ -25,7 +25,7 @@ public:
 
 	virtual void StartupModule() override
 	{
-		LLM_SCOPE_BYTAG(AudioAnalysis);
+		AUDIO_ANALYSIS_LLM_SCOPE
 		RegisterAssetActions();
 	}
 
diff --git a/Engine/Plugins/Runtime/Metasound/Source/MetasoundGraphCore/Private/MetasoundTrace.cpp b/Engine/Plugins/Runtime/Metasound/Source/MetasoundGraphCore/Private/MetasoundTrace.cpp
--- a/Engine/Plugins/Runtime/Metasound/Source/MetasoundGraphCore/Private/MetasoundTrace.cpp
+++ b/Engine/Plugins/Runtime/Metasound/Source/MetasoundGraphCore/Private/MetasoundTrace.cpp
@@ -7,4 +7,4 @@
 
 DECLARE_LLM_MEMORY_STAT(TEXT("MetaSound"), STAT_MetaSoundLLM, STATGROUP_LLMFULL);
 
-LLM_DEFINE_TAG(MetaSound, TEXT("MetaSound"), TEXT("Audio"), GET_STATFNAME(STAT_MetaSoundLLM), GET_STATFNAME(STAT_AudioSummaryLLM));
+LLM_DEFINE_TAG(Audio_MetaSound, TEXT("MetaSound"), TEXT("Audio"), GET_STATFNAME(STAT_MetaSoundLLM), GET_STATFNAME(STAT_AudioSummaryLLM));
diff --git a/Engine/Plugins/Runtime/Metasound/Source/MetasoundGraphCore/Public/MetasoundTrace.h b/Engine/Plugins/Runtime/Metasound/Source/MetasoundGraphCore/Public/MetasoundTrace.h
--- a/Engine/Plugins/Runtime/Metasound/Source/MetasoundGraphCore/Public/MetasoundTrace.h
+++ b/Engine/Plugins/Runtime/Metasound/Source/MetasoundGraphCore/Public/MetasoundTrace.h
@@ -30,8 +30,6 @@
 
 #endif
 
-LLM_DECLARE_TAG_API(MetaSound, METASOUNDGRAPHCORE_API);
-// Convenience macro for MetaSound LLM scope to avoid misspells. 
-//
-#define METASOUND_LLM_SCOPE LLM_SCOPE_BYTAG(MetaSound);
-
+LLM_DECLARE_TAG_API(Audio_MetaSound, METASOUNDGRAPHCORE_API);
+// Convenience macro for Audio_MetaSound LLM scope to avoid misspells.
+#define METASOUND_LLM_SCOPE LLM_SCOPE_BYTAG(Audio_MetaSound);
diff --git a/Engine/Plugins/Runtime/ResonanceAudio/Source/ResonanceAudio/Private/ResonanceAudioAmbisonicsSettings.cpp b/Engine/Plugins/Runtime/ResonanceAudio/Source/ResonanceAudio/Private/ResonanceAudioAmbisonicsSettings.cpp
--- a/Engine/Plugins/Runtime/ResonanceAudio/Source/ResonanceAudio/Private/ResonanceAudioAmbisonicsSettings.cpp
+++ b/Engine/Plugins/Runtime/ResonanceAudio/Source/ResonanceAudio/Private/ResonanceAudioAmbisonicsSettings.cpp
@@ -11,7 +11,7 @@
 
 TUniquePtr<ISoundfieldEncodingSettingsProxy> UResonanceAudioSoundfieldSettings::GetProxy() const
 {
-	LLM_SCOPE_BYTAG(AudioSpatializationPlugins);
+	AUDIO_SPATIALIZATION_PLUGIN_LLM_SCOPE
 	FResonanceAmbisonicsSettingsProxy* Proxy = new FResonanceAmbisonicsSettingsProxy();
 
 	switch (RenderMode)
diff --git a/Engine/Plugins/Runtime/ResonanceAudio/Source/ResonanceAudio/Private/ResonanceAudioAmbisonicsSettingsProxy.h b/Engine/Plugins/Runtime/ResonanceAudio/Source/ResonanceAudio/Private/ResonanceAudioAmbisonicsSettingsProxy.h
--- a/Engine/Plugins/Runtime/ResonanceAudio/Source/ResonanceAudio/Private/ResonanceAudioAmbisonicsSettingsProxy.h
+++ b/Engine/Plugins/Runtime/ResonanceAudio/Source/ResonanceAudio/Private/ResonanceAudioAmbisonicsSettingsProxy.h
@@ -22,7 +22,7 @@ public:
 
 	virtual TUniquePtr<ISoundfieldEncodingSettingsProxy> Duplicate() const override
 	{
-		LLM_SCOPE_BYTAG(AudioSpatializationPlugins);
+		AUDIO_SPATIALIZATION_PLUGIN_LLM_SCOPE
 		FResonanceAmbisonicsSettingsProxy* Proxy = new FResonanceAmbisonicsSettingsProxy();
 		Proxy->RenderingMode = RenderingMode;
 		return TUniquePtr<ISoundfieldEncodingSettingsProxy>(Proxy);
diff --git a/Engine/Plugins/Runtime/ResonanceAudio/Source/ResonanceAudio/Private/ResonanceAudioModule.cpp b/Engine/Plugins/Runtime/ResonanceAudio/Source/ResonanceAudio/Private/ResonanceAudioModule.cpp
--- a/Engine/Plugins/Runtime/ResonanceAudio/Source/ResonanceAudio/Private/ResonanceAudioModule.cpp
+++ b/Engine/Plugins/Runtime/ResonanceAudio/Source/ResonanceAudio/Private/ResonanceAudioModule.cpp
@@ -31,7 +31,7 @@ namespace ResonanceAudio
 
 	void FResonanceAudioModule::StartupModule()
 	{
-		LLM_SCOPE_BYTAG(AudioSpatializationPlugins);
+		LLM_SCOPE_BYTAG(Audio_SpatializationPlugins);
 		check(bModuleInitialized == false);
 		bModuleInitialized = true;
 
@@ -62,7 +62,7 @@ namespace ResonanceAudio
 
 	void FResonanceAudioModule::ShutdownModule()
 	{
-		LLM_SCOPE_BYTAG(AudioSpatializationPlugins);
+		LLM_SCOPE_BYTAG(Audio_SpatializationPlugins);
 		check(bModuleInitialized == true);
 		bModuleInitialized = false;
 		UE_LOG(LogResonanceAudio, Log, TEXT("Resonance Audio Module is Shutdown"));
@@ -85,7 +85,7 @@ namespace ResonanceAudio
 
 	void FResonanceAudioModule::RegisterAudioDevice(FAudioDevice* AudioDeviceHandle)
 	{
-		LLM_SCOPE_BYTAG(AudioSpatializationPlugins);
+		LLM_SCOPE_BYTAG(Audio_SpatializationPlugins);
 		if (!RegisteredAudioDevices.Contains(AudioDeviceHandle))
 		{
 			TAudioPluginListenerPtr NewResonanceAudioPluginListener = TAudioPluginListenerPtr(new FResonanceAudioPluginListener());
@@ -96,7 +96,7 @@ namespace ResonanceAudio
 
 	void FResonanceAudioModule::UnregisterAudioDevice(FAudioDevice* AudioDeviceHandle)
 	{
-		LLM_SCOPE_BYTAG(AudioSpatializationPlugins);
+		LLM_SCOPE_BYTAG(Audio_SpatializationPlugins);
 		RegisteredAudioDevices.Remove(AudioDeviceHandle);
 		UE_LOG(LogResonanceAudio, Log, TEXT("Audio Device unregistered from Resonance"));
 	}
@@ -108,7 +108,7 @@ namespace ResonanceAudio
 
 	TAudioSpatializationPtr FSpatializationPluginFactory::CreateNewSpatializationPlugin(FAudioDevice* OwningDevice)
 	{
-		LLM_SCOPE_BYTAG(AudioSpatializationPlugins);
+		LLM_SCOPE_BYTAG(Audio_SpatializationPlugins);
 		// Register the audio device to the Resonance Audio module.
 		FResonanceAudioModule* Module = &FModuleManager::GetModuleChecked<FResonanceAudioModule>("ResonanceAudio");
 		if (Module != nullptr)
@@ -120,7 +120,7 @@ namespace ResonanceAudio
 
 	TAudioReverbPtr FReverbPluginFactory::CreateNewReverbPlugin(FAudioDevice* OwningDevice)
 	{
-		LLM_SCOPE_BYTAG(AudioSpatializationPlugins);
+		LLM_SCOPE_BYTAG(Audio_SpatializationPlugins);
 		// Register the audio device to the Resonance Audio module.
 		FResonanceAudioModule* Module = &FModuleManager::GetModuleChecked<FResonanceAudioModule>("ResonanceAudio");
 		if (Module != nullptr)
diff --git a/Engine/Source/Runtime/AudioAnalyzer/Private/AudioAnalyzerFacade.cpp b/Engine/Source/Runtime/AudioAnalyzer/Private/AudioAnalyzerFacade.cpp
--- a/Engine/Source/Runtime/AudioAnalyzer/Private/AudioAnalyzerFacade.cpp
+++ b/Engine/Source/Runtime/AudioAnalyzer/Private/AudioAnalyzerFacade.cpp
@@ -10,7 +10,7 @@ namespace Audio
 {
 	IAnalyzerFactory* GetAnalyzerFactory(FName InFactoryName)
 	{
-		LLM_SCOPE_BYTAG(AudioAnalysis);
+		AUDIO_ANALYSIS_LLM_SCOPE
 
 		// Get all analyzer nrt factories implementations.
 		IModularFeatures::Get().LockModularFeatureList();
@@ -45,7 +45,7 @@ namespace Audio
 
 	TUniquePtr<IAnalyzerResult> FAnalyzerFacade::AnalyzeAudioBuffer(const TArray<float>& InAudioBuffer, int32 InNumChannels, float InSampleRate)
 	{
-		LLM_SCOPE_BYTAG(AudioAnalysis);
+		LLM_SCOPE_BYTAG(Audio_Analysis);
 
 		if (nullptr == Factory)
 		{
diff --git a/Engine/Source/Runtime/AudioAnalyzer/Private/AudioAnalyzerModule.cpp b/Engine/Source/Runtime/AudioAnalyzer/Private/AudioAnalyzerModule.cpp
--- a/Engine/Source/Runtime/AudioAnalyzer/Private/AudioAnalyzerModule.cpp
+++ b/Engine/Source/Runtime/AudioAnalyzer/Private/AudioAnalyzerModule.cpp
@@ -9,7 +9,7 @@
 DEFINE_LOG_CATEGORY(LogAudioAnalyzer);
 
 DECLARE_LLM_MEMORY_STAT(TEXT("AudioAnalysis"), STAT_AudioAnalysisLLM, STATGROUP_LLMFULL);
-LLM_DEFINE_TAG(AudioAnalysis, TEXT("Audio Analysis"), TEXT("Audio"), GET_STATFNAME(STAT_AudioAnalysisLLM), GET_STATFNAME(STAT_AudioSummaryLLM));
+LLM_DEFINE_TAG(Audio_Analysis, TEXT("Audio Analysis"), TEXT("Audio"), GET_STATFNAME(STAT_AudioAnalysisLLM), GET_STATFNAME(STAT_AudioSummaryLLM));
 
 void FAudioAnalyzerModule::StartupModule()
 {
diff --git a/Engine/Source/Runtime/AudioAnalyzer/Private/AudioAnalyzerNRT.cpp b/Engine/Source/Runtime/AudioAnalyzer/Private/AudioAnalyzerNRT.cpp
--- a/Engine/Source/Runtime/AudioAnalyzer/Private/AudioAnalyzerNRT.cpp
+++ b/Engine/Source/Runtime/AudioAnalyzer/Private/AudioAnalyzerNRT.cpp
@@ -122,7 +122,7 @@ bool UAudioAnalyzerNRT::ShouldEventTriggerAnalysis(struct FPropertyChangedEvent
 
 void UAudioAnalyzerNRT::AnalyzeAudio()
 {
-	LLM_SCOPE_BYTAG(AudioAnalysis);
+	AUDIO_ANALYSIS_LLM_SCOPE
 
 	TSharedPtr<Audio::IAnalyzerNRTResult, ESPMode::ThreadSafe> NewResult;
 
@@ -269,7 +269,7 @@ void UAudioAnalyzerNRT::Serialize(FArchive& Ar)
 
 TUniquePtr<Audio::IAnalyzerNRTSettings> UAudioAnalyzerNRT::GetSettings(const float InSampleRate, const int32 InNumChannels) const
 {
-	LLM_SCOPE_BYTAG(AudioAnalysis);
+	AUDIO_ANALYSIS_LLM_SCOPE
 
 	return MakeUnique<Audio::IAnalyzerNRTSettings>();
 }
diff --git a/Engine/Source/Runtime/AudioAnalyzer/Private/AudioAnalyzerNRTFacade.cpp b/Engine/Source/Runtime/AudioAnalyzer/Private/AudioAnalyzerNRTFacade.cpp
--- a/Engine/Source/Runtime/AudioAnalyzer/Private/AudioAnalyzerNRTFacade.cpp
+++ b/Engine/Source/Runtime/AudioAnalyzer/Private/AudioAnalyzerNRTFacade.cpp
@@ -53,7 +53,7 @@ namespace Audio
 	 */
 	TUniquePtr<IAnalyzerNRTResult> FAnalyzerNRTFacade::AnalyzePCM16Audio(const TArray<uint8>& InRawWaveData, int32 InNumChannels, float InSampleRate)
 	{
-		LLM_SCOPE_BYTAG(AudioAnalysis);
+		AUDIO_ANALYSIS_LLM_SCOPE
 
 		// Get analyzer factory
 		IAnalyzerNRTFactory* Factory = GetAnalyzerNRTFactory(FactoryName);
diff --git a/Engine/Source/Runtime/AudioAnalyzer/Private/AudioAnalyzerSubsystem.cpp b/Engine/Source/Runtime/AudioAnalyzer/Private/AudioAnalyzerSubsystem.cpp
--- a/Engine/Source/Runtime/AudioAnalyzer/Private/AudioAnalyzerSubsystem.cpp
+++ b/Engine/Source/Runtime/AudioAnalyzer/Private/AudioAnalyzerSubsystem.cpp
@@ -31,7 +31,7 @@ UAudioAnalyzerSubsystem* UAudioAnalyzerSubsystem::Get()
 
 bool UAudioAnalyzerSubsystem::Tick(float DeltaTime)
 {
-	LLM_SCOPE_BYTAG(AudioAnalysis);
+	AUDIO_ANALYSIS_LLM_SCOPE
 
 	// Loop through all analyzers and if they're ready to analyze, do it
 	for (UAudioAnalyzer* Analyzer : AudioAnalyzers)
@@ -54,7 +54,7 @@ void UAudioAnalyzerSubsystem::Initialize(FSubsystemCollectionBase& Collection)
 
 void UAudioAnalyzerSubsystem::Deinitialize()
 {
-	LLM_SCOPE_BYTAG(AudioAnalysis);
+	AUDIO_ANALYSIS_LLM_SCOPE
 
 	// Release our audio analyzers
 	AudioAnalyzers.Reset();
@@ -62,7 +62,7 @@ void UAudioAnalyzerSubsystem::Deinitialize()
 
 void UAudioAnalyzerSubsystem::RegisterAudioAnalyzer(UAudioAnalyzer* InAnalyzer)
 {
-	LLM_SCOPE_BYTAG(AudioAnalysis);
+	AUDIO_ANALYSIS_LLM_SCOPE
 
 	if (!TickerHandle.IsValid())
 	{
@@ -73,7 +73,7 @@ void UAudioAnalyzerSubsystem::RegisterAudioAnalyzer(UAudioAnalyzer* InAnalyzer)
 
 void UAudioAnalyzerSubsystem::UnregisterAudioAnalyzer(UAudioAnalyzer* InAnalyzer)
 {
-	LLM_SCOPE_BYTAG(AudioAnalysis);
+	AUDIO_ANALYSIS_LLM_SCOPE
 
 	AudioAnalyzers.Remove(InAnalyzer);
 	if (AudioAnalyzers.IsEmpty() && TickerHandle.IsValid())
diff --git a/Engine/Source/Runtime/AudioAnalyzer/Public/AudioAnalyzerModule.h b/Engine/Source/Runtime/AudioAnalyzer/Public/AudioAnalyzerModule.h
--- a/Engine/Source/Runtime/AudioAnalyzer/Public/AudioAnalyzerModule.h
+++ b/Engine/Source/Runtime/AudioAnalyzer/Public/AudioAnalyzerModule.h
@@ -12,7 +12,10 @@
 #include "Modules/ModuleManager.h"
 
 DECLARE_LOG_CATEGORY_EXTERN(LogAudioAnalyzer, Log, All);
-LLM_DECLARE_TAG_API(AudioAnalysis, AUDIOANALYZER_API);
+
+LLM_DECLARE_TAG_API(Audio_Analysis, AUDIOANALYZER_API);
+// Convenience macro for Audio_Analysis LLM scope to avoid misspells.
+#define AUDIO_ANALYSIS_LLM_SCOPE LLM_SCOPE_BYTAG(Audio_Analysis);
 
 class AUDIOANALYZER_API FAudioAnalyzerModule : public IModuleInterface
 {
diff --git a/Engine/Source/Runtime/Engine/Private/AudioDevice.cpp b/Engine/Source/Runtime/Engine/Private/AudioDevice.cpp
--- a/Engine/Source/Runtime/Engine/Private/AudioDevice.cpp
+++ b/Engine/Source/Runtime/Engine/Private/AudioDevice.cpp
@@ -54,9 +54,7 @@
 #endif // WITH_EDITOR
 
 DECLARE_LLM_MEMORY_STAT(TEXT("AudioSpatializationPlugins"), STAT_AudioSpatializationPluginsLLM, STATGROUP_LLMFULL);
-LLM_DEFINE_TAG(AudioSpatializationPlugins, NAME_None, TEXT("Audio"), GET_STATFNAME(STAT_AudioSpatializationPluginsLLM), GET_STATFNAME(STAT_AudioSummaryLLM));
-
-#define AUDIO_SPATIALIZATION_PLUGIN_LLM_SCOPE LLM_SCOPE_BYTAG(AudioSpatializationPlugins); 
+LLM_DEFINE_TAG(Audio_SpatializationPlugins, NAME_None, TEXT("Audio"), GET_STATFNAME(STAT_AudioSpatializationPluginsLLM), GET_STATFNAME(STAT_AudioSummaryLLM));
 
 static int32 AudioChannelCountCVar = 0;
 FAutoConsoleVariableRef CVarSetAudioChannelCount(
diff --git a/Engine/Source/Runtime/Engine/Public/AudioDevice.h b/Engine/Source/Runtime/Engine/Public/AudioDevice.h
--- a/Engine/Source/Runtime/Engine/Public/AudioDevice.h
+++ b/Engine/Source/Runtime/Engine/Public/AudioDevice.h
@@ -66,7 +66,9 @@ struct FAudioComponentParam;
 struct FAudioQualitySettings;
 struct FWaveInstance;
 
-LLM_DECLARE_TAG_API(AudioSpatializationPlugins, ENGINE_API);
+LLM_DECLARE_TAG_API(Audio_SpatializationPlugins, ENGINE_API);
+// Convenience macro for Audio_SpatializationPlugins LLM scope to avoid misspells.
+#define AUDIO_SPATIALIZATION_PLUGIN_LLM_SCOPE LLM_SCOPE_BYTAG(Audio_SpatializationPlugins);
 
 namespace Audio
 {
