Remove EIGEN_HAS_CXX11_MATH=0 from Android config

diff --git a/Engine/Plugins/Runtime/ResonanceAudio/Source/ResonanceAudio/ResonanceAudio.Build.cs b/Engine/Plugins/Runtime/ResonanceAudio/Source/ResonanceAudio/ResonanceAudio.Build.cs
--- a/Engine/Plugins/Runtime/ResonanceAudio/Source/ResonanceAudio/ResonanceAudio.Build.cs
+++ b/Engine/Plugins/Runtime/ResonanceAudio/Source/ResonanceAudio/ResonanceAudio.Build.cs
@@ -90,7 +90,6 @@ public class ResonanceAudio : ModuleRules
 		if (Target.Platform == UnrealTargetPlatform.Android)
         {
             PrivateDefinitions.Add("PFFFT_SIMD_DISABLE=1");
-            PrivateDefinitions.Add("EIGEN_HAS_CXX11_MATH=0");
         }
     }
 }
