Inclusive language push, removing master language from Resonance

diff --git a/Engine/Plugins/Runtime/ResonanceAudio/Source/ResonanceAudio/Private/ResonanceAudioLibrary/resonance_audio/api/resonance_audio_api.h b/Engine/Plugins/Runtime/ResonanceAudio/Source/ResonanceAudio/Private/ResonanceAudioLibrary/resonance_audio/api/resonance_audio_api.h
--- a/Engine/Plugins/Runtime/ResonanceAudio/Source/ResonanceAudio/Private/ResonanceAudioLibrary/resonance_audio/api/resonance_audio_api.h
+++ b/Engine/Plugins/Runtime/ResonanceAudio/Source/ResonanceAudio/Private/ResonanceAudioLibrary/resonance_audio/api/resonance_audio_api.h
@@ -205,11 +205,11 @@ class ResonanceAudioApi {
   // @param w W component of quaternion.
   virtual void SetHeadRotation(float x, float y, float z, float w) = 0;
 
-  // Sets the master volume of the main audio output.
+  // Sets the Primary volume of the main audio output.
   //
-  // @param volume Master volume (linear) in amplitude in range [0, 1] for
+  // @param volume Primary volume (linear) in amplitude in range [0, 1] for
   //     attenuation, range [1, inf) for gain boost.
-  virtual void SetMasterVolume(float volume) = 0;
+  virtual void SetPrimaryVolume(float volume) = 0;
 
   // Enables the stereo speaker mode. When activated, it disables HRTF-based
   // filtering and switches to computationally cheaper stereo-panning. This
diff --git a/Engine/Plugins/Runtime/ResonanceAudio/Source/ResonanceAudio/Private/ResonanceAudioLibrary/resonance_audio/dsp/distance_attenuation.cc b/Engine/Plugins/Runtime/ResonanceAudio/Source/ResonanceAudio/Private/ResonanceAudioLibrary/resonance_audio/dsp/distance_attenuation.cc
--- a/Engine/Plugins/Runtime/ResonanceAudio/Source/ResonanceAudio/Private/ResonanceAudioLibrary/resonance_audio/dsp/distance_attenuation.cc
+++ b/Engine/Plugins/Runtime/ResonanceAudio/Source/ResonanceAudio/Private/ResonanceAudioLibrary/resonance_audio/dsp/distance_attenuation.cc
@@ -79,7 +79,7 @@ float ComputeNearFieldEffectGain(const WorldPosition& listener_position,
   return 0.0f;
 }
 
-void UpdateAttenuationParameters(float master_gain, float reflections_gain,
+void UpdateAttenuationParameters(float primary_gain, float reflections_gain,
                                  float reverb_gain,
                                  const WorldPosition& listener_position,
                                  SourceParameters* parameters) {
@@ -106,7 +106,7 @@ void UpdateAttenuationParameters(float master_gain, float reflections_gain,
       break;
   }
   // Update gain attenuations.
-  const float input_gain = master_gain * parameters->gain;
+  const float input_gain = primary_gain * parameters->gain;
   const float direct_attenuation = input_gain * distance_attenuation;
   const float room_effects_attenuation = parameters->room_effects_gain;
 
diff --git a/Engine/Plugins/Runtime/ResonanceAudio/Source/ResonanceAudio/Private/ResonanceAudioLibrary/resonance_audio/graph/resonance_audio_api_impl.cc b/Engine/Plugins/Runtime/ResonanceAudio/Source/ResonanceAudio/Private/ResonanceAudioLibrary/resonance_audio/graph/resonance_audio_api_impl.cc
--- a/Engine/Plugins/Runtime/ResonanceAudio/Source/ResonanceAudio/Private/ResonanceAudioLibrary/resonance_audio/graph/resonance_audio_api_impl.cc
+++ b/Engine/Plugins/Runtime/ResonanceAudio/Source/ResonanceAudio/Private/ResonanceAudioLibrary/resonance_audio/graph/resonance_audio_api_impl.cc
@@ -145,8 +145,8 @@ void ResonanceAudioApiImpl::SetHeadRotation(float x, float y, float z,
   task_queue_.Post(task);
 }
 
-void ResonanceAudioApiImpl::SetMasterVolume(float volume) {
-  auto task = [this, volume]() { system_settings_.SetMasterGain(volume); };
+void ResonanceAudioApiImpl::SetPrimaryVolume(float volume) {
+  auto task = [this, volume]() { system_settings_.SetPrimaryGain(volume); };
   task_queue_.Post(task);
 }
 
@@ -468,12 +468,12 @@ void ResonanceAudioApiImpl::ProcessNextBuffer() {
   }
   // Update source attenuation parameters.
   const auto process = [this](SourceParameters* parameters) {
-    const float master_gain = system_settings_.GetMasterGain();
+    const float primary_gain = system_settings_.GetPrimaryGain();
     const auto& listener_position = system_settings_.GetHeadPosition();
     const auto& reflection_properties =
         system_settings_.GetReflectionProperties();
     const auto& reverb_properties = system_settings_.GetReverbProperties();
-    UpdateAttenuationParameters(master_gain, reflection_properties.gain,
+    UpdateAttenuationParameters(primary_gain, reflection_properties.gain,
                                 reverb_properties.gain, listener_position,
                                 parameters);
   };
diff --git a/Engine/Plugins/Runtime/ResonanceAudio/Source/ResonanceAudio/Private/ResonanceAudioLibrary/resonance_audio/graph/resonance_audio_api_impl.h b/Engine/Plugins/Runtime/ResonanceAudio/Source/ResonanceAudio/Private/ResonanceAudioLibrary/resonance_audio/graph/resonance_audio_api_impl.h
--- a/Engine/Plugins/Runtime/ResonanceAudio/Source/ResonanceAudio/Private/ResonanceAudioLibrary/resonance_audio/graph/resonance_audio_api_impl.h
+++ b/Engine/Plugins/Runtime/ResonanceAudio/Source/ResonanceAudio/Private/ResonanceAudioLibrary/resonance_audio/graph/resonance_audio_api_impl.h
@@ -60,7 +60,7 @@ class ResonanceAudioApiImpl : public ResonanceAudioApi {
   // Listener configuration.
   void SetHeadPosition(float x, float y, float z) override;
   void SetHeadRotation(float x, float y, float z, float w) override;
-  void SetMasterVolume(float volume) override;
+  void SetPrimaryVolume(float volume) override;
   void SetStereoSpeakerMode(bool enabled) override;
 
   // Create and destroy sources.
diff --git a/Engine/Plugins/Runtime/ResonanceAudio/Source/ResonanceAudio/Private/ResonanceAudioLibrary/resonance_audio/graph/system_settings.h b/Engine/Plugins/Runtime/ResonanceAudio/Source/ResonanceAudio/Private/ResonanceAudioLibrary/resonance_audio/graph/system_settings.h
--- a/Engine/Plugins/Runtime/ResonanceAudio/Source/ResonanceAudio/Private/ResonanceAudioLibrary/resonance_audio/graph/system_settings.h
+++ b/Engine/Plugins/Runtime/ResonanceAudio/Source/ResonanceAudio/Private/ResonanceAudioLibrary/resonance_audio/graph/system_settings.h
@@ -41,7 +41,7 @@ class SystemSettings {
         num_channels_(num_output_channels),
         head_rotation_(WorldRotation::Identity()),
         head_position_(WorldPosition::Zero()),
-        master_gain_(1.0f),
+        primary_gain_(1.0f),
         stereo_speaker_mode_(false) {}
 
   // Sets the listener head orientation.
@@ -110,10 +110,10 @@ class SystemSettings {
   // @return Current stereo speaker mode state.
   bool IsStereoSpeakerModeEnabled() const { return stereo_speaker_mode_; }
 
-  // Sets the master gain.
+  // Sets the primary gain.
   //
-  // @param master_gain Master output gain.
-  void SetMasterGain(float master_gain) { master_gain_ = master_gain; }
+  // @param primary_gain primary output gain.
+  void SetPrimaryGain(float primary_gain) { primary_gain_ = primary_gain; }
 
   // Sets current reflection properties.
   //
@@ -130,10 +130,10 @@ class SystemSettings {
     reverb_properties_ = reverb_properties;
   }
 
-  // Returns the master gain.
+  // Returns the primary gain.
   //
-  // @return Master output gain.
-  float GetMasterGain() const { return master_gain_; }
+  // @return primary output gain.
+  float GetPrimaryGain() const { return primary_gain_; }
 
   // Returns the current reflection properties of the environment.
   //
@@ -171,8 +171,8 @@ class SystemSettings {
   // Source parameters manager.
   SourceParametersManager source_parameters_manager_;
 
-  // Master gain in amplitude.
-  float master_gain_;
+  // primary gain in amplitude.
+  float primary_gain_;
 
   // Current reflection properties of the environment.
   ReflectionProperties reflection_properties_;
