Fix asset class path for Resonance audio global reverb preset setting

#rnx

diff --git a/Engine/Plugins/Runtime/ResonanceAudio/Source/ResonanceAudio/Private/ResonanceAudioSettings.h b/Engine/Plugins/Runtime/ResonanceAudio/Source/ResonanceAudio/Private/ResonanceAudioSettings.h
--- a/Engine/Plugins/Runtime/ResonanceAudio/Source/ResonanceAudio/Private/ResonanceAudioSettings.h
+++ b/Engine/Plugins/Runtime/ResonanceAudio/Source/ResonanceAudio/Private/ResonanceAudioSettings.h
@@ -25,7 +25,7 @@ public:
 	ERaQualityMode QualityMode;
 
 	// Default settings for global reverb: This is overridden when a player enters Audio Volumes.
-	UPROPERTY(GlobalConfig, EditAnywhere, Category = General, meta = (AllowedClasses = "/Script/ResonanceAudioEditor.ResonanceAudioReverbPluginPreset"))
+	UPROPERTY(GlobalConfig, EditAnywhere, Category = General, meta = (AllowedClasses = "/Script/ResonanceAudio.ResonanceAudioReverbPluginPreset"))
 	FSoftObjectPath GlobalReverbPreset;
 
 	// Default settings for global source settings
