Check that resonance auido reverb plugin and spatialization plugin are both resonance audio.

diff --git a/Engine/Plugins/Runtime/ResonanceAudio/Source/ResonanceAudio/Private/ResonanceAudioPluginListener.cpp b/Engine/Plugins/Runtime/ResonanceAudio/Source/ResonanceAudio/Private/ResonanceAudioPluginListener.cpp
--- a/Engine/Plugins/Runtime/ResonanceAudio/Source/ResonanceAudio/Private/ResonanceAudioPluginListener.cpp
+++ b/Engine/Plugins/Runtime/ResonanceAudio/Source/ResonanceAudio/Private/ResonanceAudioPluginListener.cpp
@@ -70,20 +70,32 @@ namespace ResonanceAudio
 			OwningAudioDevice = AudioDevice;
 		}
 
-		ReverbPtr = (FResonanceAudioReverb*)AudioDevice->ReverbPluginInterface.Get();
-		SpatializationPtr = (FResonanceAudioSpatialization*)AudioDevice->GetSpatializationPluginInterface().Get();
+		// Get the names of the currently configured spatializer / reverb
+		const FString CurrentSpatializationPluginName = AudioPluginUtilities::GetDesiredPluginName(EAudioPlugin::SPATIALIZATION);
+		const FString CurrentReverbPluginName = AudioPluginUtilities::GetDesiredPluginName(EAudioPlugin::REVERB);
 
-		// Make sure that Reverb *AND* spatialization plugins are enabled.
-		if (ReverbPtr == nullptr || SpatializationPtr == nullptr)
+		IAudioReverb* CurrentReverbPtr = AudioDevice->ReverbPluginInterface.Get();
+		IAudioSpatialization* CurrentSpatializationPtr = AudioDevice->GetSpatializationPluginInterface().Get();
+
+		// Check if reverb and spatializer are Resonance
+		bool ReverbIsResonance = CurrentReverbPtr != nullptr && CurrentReverbPluginName.Equals("Resonance Audio");
+		bool SpatializationIsResonance = CurrentSpatializationPtr != nullptr && CurrentSpatializationPluginName.Equals("Resonance Audio");
+
+		if (ReverbIsResonance && SpatializationIsResonance)
+		{
+			// Resonance is configured correctly, setup the pointers.
+			ReverbPtr = (FResonanceAudioReverb*)CurrentReverbPtr;
+			SpatializationPtr = (FResonanceAudioSpatialization*)CurrentSpatializationPtr;
+			ReverbPtr->SetResonanceAudioApi(ResonanceAudioApi.Get());
+			SpatializationPtr->SetResonanceAudioApi(ResonanceAudioApi.Get());
+			UE_LOG(LogResonanceAudio, Display, TEXT("Resonance Audio Listener is initialized"));
+		}
+		else if (CurrentReverbPluginName.Equals("Resonance Audio") || CurrentSpatializationPluginName.Equals("Resonance Audio"))
 		{
+			// Only one of reverb or spatializer was set to Resonance. Inform user that both need to be enabled
 			UE_LOG(LogResonanceAudio, Error, TEXT("Resonance Audio requires both Reverb and Spatialization plugins. Please enable them in the Project Settings."));
 			return;
 		}
-
-		ReverbPtr->SetResonanceAudioApi(ResonanceAudioApi.Get());
-		SpatializationPtr->SetResonanceAudioApi(ResonanceAudioApi.Get());
-
-		UE_LOG(LogResonanceAudio, Display, TEXT("Resonance Audio Listener is initialized"));
 	}
 
 	void FResonanceAudioPluginListener::OnListenerUpdated(FAudioDevice* AudioDevice, const int32 ViewportIndex, const FTransform& ListenerTransform, const float InDeltaSeconds)
