Fixed mismatched copy constructor/assignment operator in vraudio::ChannelView.

diff --git a/Engine/Plugins/Runtime/ResonanceAudio/Source/ResonanceAudio/Private/ResonanceAudioLibrary/resonance_audio/base/channel_view.h b/Engine/Plugins/Runtime/ResonanceAudio/Source/ResonanceAudio/Private/ResonanceAudioLibrary/resonance_audio/base/channel_view.h
--- a/Engine/Plugins/Runtime/ResonanceAudio/Source/ResonanceAudio/Private/ResonanceAudioLibrary/resonance_audio/base/channel_view.h
+++ b/Engine/Plugins/Runtime/ResonanceAudio/Source/ResonanceAudio/Private/ResonanceAudioLibrary/resonance_audio/base/channel_view.h
@@ -29,6 +29,12 @@ namespace vraudio {
 // |ChannelView| instance does not own the data it is initialized with.
 class ChannelView {
  public:
+  // @third party code - Epic Games Begin
+  // This is here for this warning as of Clang 16:
+  // channel_view.h(82,16): error: definition of implicit copy constructor for 'ChannelView' is deprecated because it has a user-provided copy assignment operator [-Werror,-Wdeprecated-copy-with-user-provided-copy]
+  ChannelView(const ChannelView&) = default;
+  // @third party code - Epic Games End
+
   // Array subscript operator returning a reference.
   float& operator[](size_t index) {
     DCHECK(enabled_);
