Make macro for unreachable code disabling
Enable C4702 on xbox
Fix every more C4702 warnings

diff --git a/Engine/Plugins/Runtime/ResonanceAudio/Source/ResonanceAudio/Private/ResonanceAudioCommon.cpp b/Engine/Plugins/Runtime/ResonanceAudio/Source/ResonanceAudio/Private/ResonanceAudioCommon.cpp
--- a/Engine/Plugins/Runtime/ResonanceAudio/Source/ResonanceAudio/Private/ResonanceAudioCommon.cpp
+++ b/Engine/Plugins/Runtime/ResonanceAudio/Source/ResonanceAudio/Private/ResonanceAudioCommon.cpp
@@ -19,49 +19,51 @@ namespace ResonanceAudio
 
 	void* LoadResonanceAudioDynamicLibrary()
 	{
-		FString LibraryPath = FPaths::EngineDir() / TEXT("Source/ThirdParty/ResonanceAudioApi/lib");
-		FString DynamicLibraryToLoad;
-		void* DynamicLibraryHandle = nullptr;
+		auto LoadLibrary = [](const FString& FilePath)
+		{
+			FString DynamicLibraryToLoad = FPaths::EngineDir() / TEXT("Source/ThirdParty/ResonanceAudioApi/lib") / FilePath;
+			void* DynamicLibraryHandle = nullptr;
+
+			UE_LOG(LogResonanceAudio, Log, TEXT("Attempting to load %s"), *DynamicLibraryToLoad);
+
+			if (FPaths::FileExists(DynamicLibraryToLoad))
+			{
+				DynamicLibraryHandle = FPlatformProcess::GetDllHandle(*DynamicLibraryToLoad);
+			}
+			else
+			{
+				UE_LOG(LogResonanceAudio, Log, TEXT("File does not exist. %s"), *DynamicLibraryToLoad);
+			}
+	
+			if (!DynamicLibraryHandle)
+			{
+				UE_LOG(LogResonanceAudio, Log, TEXT("Unable to load %s."), *FPaths::ConvertRelativePathToFull(DynamicLibraryToLoad));
+			}
+			else
+			{
+				UE_LOG(LogResonanceAudio, Log, TEXT("Loaded %s."), *DynamicLibraryToLoad);
+			}
+	
+			return DynamicLibraryHandle;
+		};
 
 #if PLATFORM_WINDOWS
 	#if PLATFORM_64BITS
-		DynamicLibraryToLoad = LibraryPath / TEXT("win_x64/vraudio.dll");
+		return LoadLibrary(TEXT("win_x64/vraudio.dll"));
 	#else
-		DynamicLibraryToLoad = LibraryPath / TEXT("win_x86/vraudio.dll");
+		return LoadLibrary(TEXT("win_x86/vraudio.dll"));
 	#endif	// PLATFORM_64BITS
 #elif PLATFORM_MAC
-		DynamicLibraryToLoad = LibraryPath / TEXT("darwin/libvraudio.dylib");
+		return LoadLibrary(TEXT("darwin/libvraudio.dylib"));
 #elif PLATFORM_ANDROID || PLATFORM_IOS
 		 // Not necessary on this platform.
 		return nullptr;
 #elif PLATFORM_LINUX
-		DynamicLibraryToLoad = LibraryPath / TEXT("linux/libvraudio.so");
+		return LoadLibrary(TEXT("linux/libvraudio.so"));
 #else
 		UE_LOG(LogResonanceAudio, Error, TEXT("Unsupported Platform. Supported platforms are ANDROID, IOS, LINUX, MAC and WINDOWS"));
 		return nullptr;
 #endif  // PLATFORM_WINDOWS
-
-		UE_LOG(LogResonanceAudio, Log, TEXT("Attempting to load %s"), *DynamicLibraryToLoad);
-
-		if (FPaths::FileExists(DynamicLibraryToLoad))
-		{
-			DynamicLibraryHandle = FPlatformProcess::GetDllHandle(*DynamicLibraryToLoad);
-		}
-		else
-		{
-			UE_LOG(LogResonanceAudio, Log, TEXT("File does not exist. %s"), *DynamicLibraryToLoad);
-		}
-
-		if (!DynamicLibraryHandle)
-		{
-			UE_LOG(LogResonanceAudio, Log, TEXT("Unable to load %s."), *FPaths::ConvertRelativePathToFull(DynamicLibraryToLoad));
-		}
-		else
-		{
-			UE_LOG(LogResonanceAudio, Log, TEXT("Loaded %s."), *DynamicLibraryToLoad);
-		}
-
-		return DynamicLibraryHandle;
 	}
 
 	vraudio::ResonanceAudioApi* CreateResonanceAudioApi(void* DynamicLibraryHandle, size_t NumChannels, size_t NumFrames, int SampleRate) {
diff --git a/Engine/Plugins/Tests/CQTest/Source/CQTestTests/Private/Assert/AssertionTests.cpp b/Engine/Plugins/Tests/CQTest/Source/CQTestTests/Private/Assert/AssertionTests.cpp
--- a/Engine/Plugins/Tests/CQTest/Source/CQTestTests/Private/Assert/AssertionTests.cpp
+++ b/Engine/Plugins/Tests/CQTest/Source/CQTestTests/Private/Assert/AssertionTests.cpp
@@ -19,15 +19,14 @@ namespace CQTests
 			Assert.Fail(AnyError);
 		}
 
-#pragma warning(push)
-#pragma warning(disable:4702)
+		PRAGMA_DISABLE_UNREACHABLE_CODE_WARNINGS
 		TEST_METHOD(AssertFail_Macro_AddsErrorAndExits)
 		{
 			Assert.ExpectError(ExpectedError);
 			ASSERT_FAIL(ExpectedError);
 			ASSERT_FAIL(TEXT("It should not assert after the first expected assert since it exits"));
 		}
-#pragma warning(pop)
+		PRAGMA_RESTORE_UNREACHABLE_CODE_WARNINGS
 
 		TEST_METHOD(Assertions_Accept_RawStrings)
 		{
diff --git a/Engine/Source/Programs/AutomationTool/Scripts/CheckBalancedMacros.cs b/Engine/Source/Programs/AutomationTool/Scripts/CheckBalancedMacros.cs
--- a/Engine/Source/Programs/AutomationTool/Scripts/CheckBalancedMacros.cs
+++ b/Engine/Source/Programs/AutomationTool/Scripts/CheckBalancedMacros.cs
@@ -49,6 +49,10 @@ namespace AutomationTool
 				"PRAGMA_DISABLE_UNSAFE_TYPECAST_WARNINGS",
 				"PRAGMA_RESTORE_UNSAFE_TYPECAST_WARNINGS"
 			},
+			{
+				"PRAGMA_DISABLE_UNREACHABLE_CODE_WARNINGS",
+				"PRAGMA_RESTORE_UNREACHABLE_CODE_WARNINGS"
+			},
 			{
 				"PRAGMA_FORCE_UNSAFE_TYPECAST_WARNINGS",
 				"PRAGMA_RESTORE_UNSAFE_TYPECAST_WARNINGS"
diff --git a/Engine/Source/Programs/TestPAL/Private/Main.cpp b/Engine/Source/Programs/TestPAL/Private/Main.cpp
--- a/Engine/Source/Programs/TestPAL/Private/Main.cpp
+++ b/Engine/Source/Programs/TestPAL/Private/Main.cpp
@@ -1439,8 +1439,7 @@ namespace
 		}
 	}
 
-#pragma warning(push)
-#pragma warning(disable:4702)
+	PRAGMA_DISABLE_UNREACHABLE_CODE_WARNINGS
 	void FORCENOINLINE LabelGoto()
 	{
 		goto end;
@@ -1453,7 +1452,7 @@ namespace
 end:
 		ensure(false);
 	}
-#pragma warning(pop)
+	PRAGMA_RESTORE_UNREACHABLE_CODE_WARNINGS
 
 	void FORCEINLINE inline_three_ensures()
 	{
diff --git a/Engine/Source/Runtime/Core/Private/GenericPlatform/GenericPlatformMallocCrash.cpp b/Engine/Source/Runtime/Core/Private/GenericPlatform/GenericPlatformMallocCrash.cpp
--- a/Engine/Source/Runtime/Core/Private/GenericPlatform/GenericPlatformMallocCrash.cpp
+++ b/Engine/Source/Runtime/Core/Private/GenericPlatform/GenericPlatformMallocCrash.cpp
@@ -454,8 +454,7 @@ void FGenericPlatformMallocCrash::PrintPoolsUsage()
 #endif // _DEBUG
 }
 
-#pragma warning(push)
-#pragma warning(disable:4702)
+PRAGMA_DISABLE_UNREACHABLE_CODE_WARNINGS
 bool FGenericPlatformMallocCrash::IsOnCrashedThread() const
 {
 	// Suspend threads other than the crashed one to prevent serious memory errors.
@@ -470,7 +469,7 @@ bool FGenericPlatformMallocCrash::IsOnCrashedThread() const
 		return false;
 	}
 }
-#pragma warning(pop)
+PRAGMA_RESTORE_UNREACHABLE_CODE_WARNINGS
 
 
 bool FGenericPlatformMallocCrash::IsPtrInLargePool( void* Ptr ) const
diff --git a/Engine/Source/Runtime/Core/Private/HAL/MallocAnsi.cpp b/Engine/Source/Runtime/Core/Private/HAL/MallocAnsi.cpp
--- a/Engine/Source/Runtime/Core/Private/HAL/MallocAnsi.cpp
+++ b/Engine/Source/Runtime/Core/Private/HAL/MallocAnsi.cpp
@@ -243,8 +243,6 @@ bool FMallocAnsi::ValidateHeap()
 	check(Result != _HEAPBADPTR);
 	check(Result != _HEAPEMPTY);
 	check(Result == _HEAPOK);
-#else
-	return true;
 #endif
 	return true;
 }
diff --git a/Engine/Source/Runtime/Core/Public/AutoRTFM/AutoRTFM.h b/Engine/Source/Runtime/Core/Public/AutoRTFM/AutoRTFM.h
--- a/Engine/Source/Runtime/Core/Public/AutoRTFM/AutoRTFM.h
+++ b/Engine/Source/Runtime/Core/Public/AutoRTFM/AutoRTFM.h
@@ -240,8 +240,7 @@ UE_AUTORTFM_FORCEINLINE void autortfm_open(void (*work)(void* arg), void* arg)
 #if UE_AUTORTFM
 [[nodiscard]] UE_AUTORTFM_API autortfm_status autortfm_close(void (*work)(void* arg), void* arg);
 #else
-#pragma warning(push)
-#pragma warning(disable:4702)
+PRAGMA_DISABLE_UNREACHABLE_CODE_WARNINGS
 [[nodiscard]] UE_AUTORTFM_FORCEINLINE autortfm_status autortfm_close(void (*work)(void* arg), void* arg)
 {
 	UE_AUTORTFM_UNUSED(work);
@@ -249,7 +248,7 @@ UE_AUTORTFM_FORCEINLINE void autortfm_open(void (*work)(void* arg), void* arg)
     abort();
 	return autortfm_status_aborted_by_language;
 }
-#pragma warning(pop)
+PRAGMA_RESTORE_UNREACHABLE_CODE_WARNINGS
 #endif
 
 // Records the pointer and size from the open into the current transaction
diff --git a/Engine/Source/Runtime/Core/Public/Clang/ClangPlatformCompilerPreSetup.h b/Engine/Source/Runtime/Core/Public/Clang/ClangPlatformCompilerPreSetup.h
--- a/Engine/Source/Runtime/Core/Public/Clang/ClangPlatformCompilerPreSetup.h
+++ b/Engine/Source/Runtime/Core/Public/Clang/ClangPlatformCompilerPreSetup.h
@@ -52,6 +52,15 @@
 		_Pragma("clang diagnostic pop")
 #endif // PRAGMA_ENABLE_SHADOW_VARIABLE_WARNINGS
 
+// Unreachable not being evaluated on clang... yet
+#ifndef PRAGMA_DISABLE_UNREACHABLE_CODE_WARNINGS
+	#define PRAGMA_DISABLE_UNREACHABLE_CODE_WARNINGS
+#endif // PRAGMA_DISABLE_UNREACHABLE_CODE_WARNINGS
+
+#ifndef PRAGMA_RESTORE_UNREACHABLE_CODE_WARNINGS
+	#define PRAGMA_RESTORE_UNREACHABLE_CODE_WARNINGS
+#endif // PRAGMA_RESTORE_UNREACHABLE_CODE_WARNINGS
+
 #if __has_warning("-Wimplicit-float-conversion")
 #define DISABLE_IMPLICIT_FLOAT_CONVERSION_FRAGMENT _Pragma("clang diagnostic ignored \"-Wimplicit-float-conversion\"")
 #else
@@ -267,6 +276,7 @@
 		PRAGMA_DISABLE_OVERLOADED_VIRTUAL_WARNINGS \
 		PRAGMA_DISABLE_MISSING_BRACES_WARNINGS \
 		PRAGMA_DISABLE_UNINITIALIZED_CONST_REFERENCE_WARNINGS \
+		PRAGMA_DISABLE_UNREACHABLE_CODE_WARNINGS \
 		PRAGMA_DISABLE_ORDERED_COMPARE_FUNCTION_POINTERS \
 		PRAGMA_DISABLE_DEPRECATED_COPY_WITH_USER_DEFINED_COPY_WARNINGS
 #endif
@@ -275,6 +285,7 @@
 	#define UE_COMPILER_THIRD_PARTY_INCLUDES_END \
 		PRAGMA_ENABLE_DEPRECATED_COPY_WITH_USER_DEFINED_COPY_WARNINGS \
 		PRAGMA_ENABLE_ORDERED_COMPARE_FUNCTION_POINTERS \
+		PRAGMA_RESTORE_UNREACHABLE_CODE_WARNINGS \
 		PRAGMA_ENABLE_UNINITIALIZED_CONST_REFERENCE_WARNINGS \
 		PRAGMA_ENABLE_MISSING_BRACES_WARNINGS \
 		PRAGMA_ENABLE_OVERLOADED_VIRTUAL_WARNINGS \
diff --git a/Engine/Source/Runtime/Core/Public/MSVC/MSVCPlatformCompilerPreSetup.h b/Engine/Source/Runtime/Core/Public/MSVC/MSVCPlatformCompilerPreSetup.h
--- a/Engine/Source/Runtime/Core/Public/MSVC/MSVCPlatformCompilerPreSetup.h
+++ b/Engine/Source/Runtime/Core/Public/MSVC/MSVCPlatformCompilerPreSetup.h
@@ -66,6 +66,17 @@
 			__pragma(warning(pop))
 	#endif // PRAGMA_ENABLE_SHADOW_VARIABLE_WARNINGS
 
+	#ifndef PRAGMA_DISABLE_UNREACHABLE_CODE_WARNINGS
+		#define PRAGMA_DISABLE_UNREACHABLE_CODE_WARNINGS \
+			__pragma (warning(push)) \
+			__pragma (warning(disable: 4702)) /* unreachable code */
+	#endif // PRAGMA_DISABLE_UNREACHABLE_CODE_WARNINGS
+	
+	#ifndef PRAGMA_RESTORE_UNREACHABLE_CODE_WARNINGS
+		#define PRAGMA_RESTORE_UNREACHABLE_CODE_WARNINGS \
+			__pragma (warning(pop))
+	#endif // PRAGMA_RESTORE_UNREACHABLE_CODE_WARNINGS
+
 	#ifndef PRAGMA_DISABLE_UNSAFE_TYPECAST_WARNINGS
 		#define PRAGMA_DISABLE_UNSAFE_TYPECAST_WARNINGS \
 			__pragma (warning(push)) \
@@ -195,7 +206,6 @@
 			__pragma(warning(disable: 4125))  /* decimal digit terminates octal escape sequence. */ \
 			__pragma(warning(disable: 4510))  /* '<class>': default constructor could not be generated. */ \
 			__pragma(warning(disable: 4610))  /* object '<class>' can never be instantiated - user-defined constructor required. */ \
-			__pragma(warning(disable: 4702))  /* unreachable code. */ \
 			__pragma(warning(disable: 4800))  /* Implicit conversion from '<type>' to bool. Possible information loss. */ \
 			__pragma(warning(disable: 4946))  /* reinterpret_cast used between related classes: '<class1>' and '<class2>' */ \
 			__pragma(warning(disable: 4996))  /* '<obj>' was declared deprecated. */ \
@@ -218,11 +228,13 @@
 			PRAGMA_DISABLE_SHADOW_VARIABLE_WARNINGS \
 			PRAGMA_DISABLE_MISSING_VIRTUAL_DESTRUCTOR_WARNINGS \
 			PRAGMA_DISABLE_DEPRECATION_WARNINGS \
-			PRAGMA_DISABLE_UNINITIALIZED_CONST_REFERENCE_WARNINGS
+			PRAGMA_DISABLE_UNINITIALIZED_CONST_REFERENCE_WARNINGS \
+			PRAGMA_DISABLE_UNREACHABLE_CODE_WARNINGS
 	#endif // THIRD_PARTY_INCLUDES_START
 
 	#ifndef UE_COMPILER_THIRD_PARTY_INCLUDES_END
 		#define UE_COMPILER_THIRD_PARTY_INCLUDES_END \
+			PRAGMA_RESTORE_UNREACHABLE_CODE_WARNINGS \
 			PRAGMA_ENABLE_UNINITIALIZED_CONST_REFERENCE_WARNINGS \
 			PRAGMA_ENABLE_DEPRECATION_WARNINGS \
 			PRAGMA_ENABLE_MISSING_VIRTUAL_DESTRUCTOR_WARNINGS \
diff --git a/Engine/Source/Runtime/Engine/Private/UnrealEngine.cpp b/Engine/Source/Runtime/Engine/Private/UnrealEngine.cpp
--- a/Engine/Source/Runtime/Engine/Private/UnrealEngine.cpp
+++ b/Engine/Source/Runtime/Engine/Private/UnrealEngine.cpp
@@ -10614,8 +10614,7 @@ bool UEngine::PerformError(const TCHAR* Cmd, FOutputDevice& Ar)
 		ENQUEUE_RENDER_COMMAND(CauseGpuCrash)(&FRender::GpuCrash);
 		return true;
 	}
-#pragma warning(push)
-#pragma warning(disable:4702)
+PRAGMA_DISABLE_UNREACHABLE_CODE_WARNINGS
 	else if (FParse::Command(&Cmd, TEXT("TERMINATE")))
 	{
 		UE_LOG(LogEngine, Warning, TEXT("Printed warning to log."));
@@ -10632,7 +10631,7 @@ bool UEngine::PerformError(const TCHAR* Cmd, FOutputDevice& Ar)
 		std::abort();
 		return true;
 	}
-#pragma warning(pop)
+PRAGMA_RESTORE_UNREACHABLE_CODE_WARNINGS
 	else if (FParse::Command(&Cmd, TEXT("CHECK")))
 	{
 		UE_LOG(LogEngine, Warning, TEXT("Printed warning to log."));
@@ -10965,8 +10964,7 @@ bool UEngine::PerformError(const TCHAR* Cmd, FOutputDevice& Ar)
 		);
 		return true;
 	}
-#pragma warning(push)
-#pragma warning(disable:4702)
+PRAGMA_DISABLE_UNREACHABLE_CODE_WARNINGS
 	else if (FParse::Command(&Cmd, TEXT("SOFTLOCK")))
 	{
 		Ar.Log(TEXT("Hanging the current thread"));
@@ -10986,7 +10984,7 @@ bool UEngine::PerformError(const TCHAR* Cmd, FOutputDevice& Ar)
 		}
 		return true;
 	}
-#pragma warning(pop)
+PRAGMA_RESTORE_UNREACHABLE_CODE_WARNINGS
 	else if (FParse::Command(&Cmd, TEXT("SLEEP")))
 	{
 		Ar.Log(TEXT("Sleep for 1 hour. This should crash after a few seconds in cooked builds."));
