Load Resonannce Plugin earlier

diff --git a/Engine/Plugins/Runtime/ResonanceAudio/ResonanceAudio.uplugin b/Engine/Plugins/Runtime/ResonanceAudio/ResonanceAudio.uplugin
--- a/Engine/Plugins/Runtime/ResonanceAudio/ResonanceAudio.uplugin
+++ b/Engine/Plugins/Runtime/ResonanceAudio/ResonanceAudio.uplugin
@@ -18,30 +18,30 @@
 	{
 		"Name": "ResonanceAudio",
 		"Type": "Runtime",
-		"LoadingPhase": "Default",
+		"LoadingPhase": "PreDefault",
 		"PlatformAllowList": [
 			"Android",
 			"IOS",
 			"Linux",
 			"Mac",
-			"Win64",
+			"Win64"
 		]
 	},
-    {
-      "Name": "ResonanceAudioEditor",
-      "Type": "Editor",
-      "LoadingPhase": "PostEngineInit",
-      "PlatformAllowList": [
-        "Linux",
-        "Mac",
-        "Win64"
-      ]
-    }
+	{
+		"Name": "ResonanceAudioEditor",
+		"Type": "Editor",
+		"LoadingPhase": "PostEngineInit",
+		"PlatformAllowList": [
+			"Linux",
+			"Mac",
+			"Win64"
+		]
+	}
   ],
   "Plugins": [
-    {
-      "Name": "ProceduralMeshComponent",
-      "Enabled": true
-    }
+	{
+		"Name": "ProceduralMeshComponent",
+		"Enabled": true
+	}
   ]
 }
