Fix for google resonance spatialization being incompatible w/ multi-pie.

diff --git a/Engine/Plugins/Runtime/ResonanceAudio/Source/ResonanceAudio/Private/ResonanceAudioReverb.cpp b/Engine/Plugins/Runtime/ResonanceAudio/Source/ResonanceAudio/Private/ResonanceAudioReverb.cpp
--- a/Engine/Plugins/Runtime/ResonanceAudio/Source/ResonanceAudio/Private/ResonanceAudioReverb.cpp
+++ b/Engine/Plugins/Runtime/ResonanceAudio/Source/ResonanceAudio/Private/ResonanceAudioReverb.cpp
@@ -29,20 +29,12 @@ namespace ResonanceAudio
 		: ResonanceAudioApi(nullptr)
 		, ResonanceAudioModule(nullptr)
 		, ReverbPluginPreset(nullptr)
-		, GlobalReverbPluginPreset(nullptr)
 		, TemporaryStereoBuffer()
 	{
 	}
 
 	FResonanceAudioReverb::~FResonanceAudioReverb()
 	{
-		if (GlobalReverbPluginPreset)
-		{
-#if !WITH_EDITOR
-			GlobalReverbPluginPreset->RemoveFromRoot();
-#endif
-			GlobalReverbPluginPreset = nullptr;
-		}
 		ReverbPluginPreset = nullptr;
 	}
 
@@ -112,6 +104,7 @@ namespace ResonanceAudio
 		}
 	}
 
+	static UResonanceAudioReverbPluginPreset* GlobalReverbPluginPreset = nullptr;
 	void FResonanceAudioReverb::SetGlobalReverbPluginPreset(UResonanceAudioReverbPluginPreset* InPreset)
 	{
 		//If we're using the global reverb preset right now, ensure that we update ReverbPluginPreset as well.
@@ -126,26 +119,39 @@ namespace ResonanceAudio
 		}
 	}
 
+	UResonanceAudioReverbPluginPreset* FResonanceAudioReverb::GetGlobalReverbPluginPreset()
+	{
+		return GlobalReverbPluginPreset;
+	}
+
 	void FResonanceAudioReverb::InitEffectSubmix()
 	{
 		check(!SubmixEffect.IsValid());
 
-		// Load the global reverb preset settings:
-		const FSoftObjectPath ReverbPluginPresetName = GetDefault<UResonanceAudioSettings>()->GlobalReverbPreset;
-		if (ReverbPluginPresetName.IsValid())
-		{
-			GlobalReverbPluginPreset = LoadObject<UResonanceAudioReverbPluginPreset>(nullptr, *ReverbPluginPresetName.ToString());
-		}
-
-		// If loading of the Reverb Plugin Preset asset fails, create a temporary preset. No reverb will be applied.
+		// initialize the GlobalReverbPluginPreset
 		if (GlobalReverbPluginPreset == nullptr)
 		{
-			GlobalReverbPluginPreset = NewObject<UResonanceAudioReverbPluginPreset>(UResonanceAudioReverbPluginPreset::StaticClass(), TEXT("Resonance Audio Reverb Plugin Preset"));
+			// Load the global reverb preset settings:
+			const FSoftObjectPath ReverbPluginPresetName = GetDefault<UResonanceAudioSettings>()->GlobalReverbPreset;
+			if (ReverbPluginPresetName.IsValid())
+			{
+				GlobalReverbPluginPreset = LoadObject<UResonanceAudioReverbPluginPreset>(nullptr, *ReverbPluginPresetName.ToString());
+			}
+
+			// If loading of the Reverb Plugin Preset asset fails, create a temporary preset.
+			if (GlobalReverbPluginPreset == nullptr)
+			{
+				GlobalReverbPluginPreset = NewObject<UResonanceAudioReverbPluginPreset>();
+			}
+			
+			if (GlobalReverbPluginPreset)
+			{
+				GlobalReverbPluginPreset->AddToRoot();
+			}
 		}
 
 		if (GlobalReverbPluginPreset)
 		{
-			GlobalReverbPluginPreset->AddToRoot();
 			ReverbPluginPreset = GlobalReverbPluginPreset;
 
 			SubmixEffect = USoundEffectPreset::CreateInstance<FSoundEffectSubmixInitData, FSoundEffectSubmix>(FSoundEffectSubmixInitData(), *ReverbPluginPreset);
@@ -154,6 +160,7 @@ namespace ResonanceAudio
 		}
 		else
 		{
+			UE_LOG(LogResonanceAudio, Error, TEXT("Failed to create Resonance Reverb global Preset"));
 			ReverbPluginPreset = nullptr;
 		}
 	}
@@ -187,7 +194,7 @@ namespace ResonanceAudio
 				*Settings->OutputSubmix.GetAssetPathString(),
 				*DefaultSubmixName);
 
-			ReverbSubmix = NewObject<USoundSubmix>(USoundSubmix::StaticClass(), *DefaultSubmixName);
+			ReverbSubmix = NewObject<USoundSubmix>();
 			ReverbSubmix->bMuteWhenBackgrounded = true;
 		}
 
diff --git a/Engine/Plugins/Runtime/ResonanceAudio/Source/ResonanceAudio/Private/ResonanceAudioReverb.h b/Engine/Plugins/Runtime/ResonanceAudio/Source/ResonanceAudio/Private/ResonanceAudioReverb.h
--- a/Engine/Plugins/Runtime/ResonanceAudio/Source/ResonanceAudio/Private/ResonanceAudioReverb.h
+++ b/Engine/Plugins/Runtime/ResonanceAudio/Source/ResonanceAudio/Private/ResonanceAudioReverb.h
@@ -121,7 +121,7 @@ namespace ResonanceAudio
 		void ProcessMixedAudio(const FSoundEffectSubmixInputData& InData, FSoundEffectSubmixOutputData& OutData);
 		
 		void SetGlobalReverbPluginPreset(UResonanceAudioReverbPluginPreset* InPreset);
-		UResonanceAudioReverbPluginPreset* GetGlobalReverbPluginPreset() { return GlobalReverbPluginPreset; };
+		UResonanceAudioReverbPluginPreset* GetGlobalReverbPluginPreset();
 
 		// Sets & updates Room Effect params from current |ReverbPluginPreset|.
 		void UpdateRoomEffects();
@@ -147,7 +147,6 @@ namespace ResonanceAudio
 		vraudio::ResonanceAudioApi* ResonanceAudioApi;
 		FResonanceAudioModule* ResonanceAudioModule;
 		UResonanceAudioReverbPluginPreset* ReverbPluginPreset;
-		UResonanceAudioReverbPluginPreset* GlobalReverbPluginPreset;
 		Audio::AlignedFloatBuffer TemporaryStereoBuffer;
 
 		FSoundEffectSubmixPtr SubmixEffect;
