Stop Logspam and Memory ballooning when Resonance Submix doesn't load properly

diff --git a/Engine/Plugins/Runtime/ResonanceAudio/Source/ResonanceAudio/Private/ResonanceAudioReverb.cpp b/Engine/Plugins/Runtime/ResonanceAudio/Source/ResonanceAudio/Private/ResonanceAudioReverb.cpp
--- a/Engine/Plugins/Runtime/ResonanceAudio/Source/ResonanceAudio/Private/ResonanceAudioReverb.cpp
+++ b/Engine/Plugins/Runtime/ResonanceAudio/Source/ResonanceAudio/Private/ResonanceAudioReverb.cpp
@@ -30,6 +30,7 @@ namespace ResonanceAudio
 		, ResonanceAudioModule(nullptr)
 		, ReverbPluginPreset(nullptr)
 		, TemporaryStereoBuffer()
+		, StubSubmixPtr(nullptr)
 	{
 	}
 
@@ -189,13 +190,21 @@ namespace ResonanceAudio
 
 		if (!ReverbSubmix)
 		{
-			static const FString DefaultSubmixName = TEXT("Resonance Reverb Submix");
-			UE_LOG(LogResonanceAudio, Error, TEXT("Failed to load Resonance Reverb Submix from object path '%s' in ResonanceSettings. Creating '%s' as stub."),
-				*Settings->OutputSubmix.GetAssetPathString(),
-				*DefaultSubmixName);
+			if (!StubSubmixPtr)
+			{
+				static const FString DefaultSubmixName = TEXT("Resonance Reverb Submix");
+				UE_LOG(LogResonanceAudio, Error, TEXT("Failed to load Resonance Reverb Submix from object path '%s' in ResonanceSettings. Creating '%s' as stub."),
+					*Settings->OutputSubmix.GetAssetPathString(),
+					*DefaultSubmixName);
+
+				StubSubmixPtr = NewObject<USoundSubmix>();
+				StubSubmixPtr->AddToRoot();
+				StubSubmixPtr->bMuteWhenBackgrounded = true;
+
+			}
+
+			ReverbSubmix = StubSubmixPtr;
 
-			ReverbSubmix = NewObject<USoundSubmix>();
-			ReverbSubmix->bMuteWhenBackgrounded = true;
 		}
 
 		ReverbSubmix->bAutoDisable = false;
diff --git a/Engine/Plugins/Runtime/ResonanceAudio/Source/ResonanceAudio/Private/ResonanceAudioReverb.h b/Engine/Plugins/Runtime/ResonanceAudio/Source/ResonanceAudio/Private/ResonanceAudioReverb.h
--- a/Engine/Plugins/Runtime/ResonanceAudio/Source/ResonanceAudio/Private/ResonanceAudioReverb.h
+++ b/Engine/Plugins/Runtime/ResonanceAudio/Source/ResonanceAudio/Private/ResonanceAudioReverb.h
@@ -150,6 +150,9 @@ namespace ResonanceAudio
 		Audio::AlignedFloatBuffer TemporaryStereoBuffer;
 
 		FSoundEffectSubmixPtr SubmixEffect;
+
+		// Stub Reverb submix to be used only when Resonance fails to load the submix asset.
+		USoundSubmix* StubSubmixPtr;
 	};
 
 } // namespace ResonanceAudio
