New CVar to force Spatialization Plugin Submix to send directly to main submix

diff --git a/Engine/Plugins/Runtime/ResonanceAudio/Source/ResonanceAudio/Private/ResonanceAudioReverb.cpp b/Engine/Plugins/Runtime/ResonanceAudio/Source/ResonanceAudio/Private/ResonanceAudioReverb.cpp
--- a/Engine/Plugins/Runtime/ResonanceAudio/Source/ResonanceAudio/Private/ResonanceAudioReverb.cpp
+++ b/Engine/Plugins/Runtime/ResonanceAudio/Source/ResonanceAudio/Private/ResonanceAudioReverb.cpp
@@ -4,6 +4,7 @@
 
 #include "ResonanceAudioReverb.h"
 
+#include "AudioDeviceManager.h"
 #include "ResonanceAudioCommon.h"
 #include "ResonanceAudioModule.h"
 #include "ResonanceAudioSettings.h"
@@ -30,13 +31,20 @@ namespace ResonanceAudio
 		, ResonanceAudioModule(nullptr)
 		, ReverbPluginPreset(nullptr)
 		, TemporaryStereoBuffer()
-		, StubSubmixPtr(nullptr)
+		, ReverbSubmixPtr(nullptr)
 	{
 	}
 
 	FResonanceAudioReverb::~FResonanceAudioReverb()
 	{
 		ReverbPluginPreset = nullptr;
+
+		if (ReverbSubmixPtr)
+		{
+			ReverbSubmixPtr->SetParentSubmix(nullptr, false /* bModifyAssets */);
+			ReverbSubmixPtr->RemoveFromRoot();
+		}
+
 	}
 
 	void FResonanceAudioReverb::Initialize(const FAudioPluginInitializationParams InitializationParams)
@@ -178,36 +186,49 @@ namespace ResonanceAudio
 
 	USoundSubmix* FResonanceAudioReverb::GetSubmix()
 	{
+
+		if (ReverbSubmixPtr)
+		{
+			return ReverbSubmixPtr;
+		}
+
 		const UResonanceAudioSettings* Settings = GetDefault<UResonanceAudioSettings>();
 		check(Settings);
 
-		USoundSubmix* ReverbSubmix = nullptr;
-		
 		if (!IgnoreUserResonanceSubmixCVar)
 		{
-			ReverbSubmix = Cast<USoundSubmix>(Settings->OutputSubmix.TryLoad());
+			ReverbSubmixPtr = Cast<USoundSubmix>(Settings->OutputSubmix.TryLoad());
 		}
 
-		if (!ReverbSubmix)
+		if (!ReverbSubmixPtr)
 		{
-			if (!StubSubmixPtr)
+			if (!IgnoreUserResonanceSubmixCVar)
 			{
 				static const FString DefaultSubmixName = TEXT("Resonance Reverb Submix");
-				UE_LOG(LogResonanceAudio, Error, TEXT("Failed to load Resonance Reverb Submix from object path '%s' in ResonanceSettings. Creating '%s' as stub."),
+				UE_LOG(LogResonanceAudio, Error, TEXT("Couldn't load Resonance Reverb Submix from object path '%s' in ResonanceSettings. Creating stub submix."),
 					*Settings->OutputSubmix.GetAssetPathString(),
 					*DefaultSubmixName);
+			}
 
-				StubSubmixPtr = NewObject<USoundSubmix>();
-				StubSubmixPtr->AddToRoot();
-				StubSubmixPtr->bMuteWhenBackgrounded = true;
+			ReverbSubmixPtr = NewObject<USoundSubmix>();
+			ReverbSubmixPtr->AddToRoot();
+			ReverbSubmixPtr->bMuteWhenBackgrounded = true;
 
-			}
+			if (const FAudioDeviceManager* DeviceManager = FAudioDeviceManager::Get())
+			{
+				if (const FAudioDevice* MainDevice = DeviceManager->GetMainAudioDeviceRaw())
+				{
+					USoundSubmix* MainSubmixObjectPtr = &MainDevice->GetMainSubmixObject();
+					UE_LOG(LogResonanceAudio, Log, TEXT("Setting Resonance's submix parent to Main Submix %s")
+						, *MainSubmixObjectPtr->GetName());
 
-			ReverbSubmix = StubSubmixPtr;
+					ReverbSubmixPtr->SetParentSubmix(MainSubmixObjectPtr, false /* bModifyAssets */);
+				}
+			}
 
 		}
 
-		ReverbSubmix->bAutoDisable = false;
+		ReverbSubmixPtr->bAutoDisable = false;
 
 		// SubmixEffect is required to be initialized for ReverbPluginPreset to be set
 		if (!SubmixEffect.IsValid())
@@ -218,7 +239,7 @@ namespace ResonanceAudio
 		if (ReverbPluginPreset)
 		{
 			bool bFoundPreset = false;
-			for (USoundEffectSubmixPreset* Preset : ReverbSubmix->SubmixEffectChain)
+			for (USoundEffectSubmixPreset* Preset : ReverbSubmixPtr->SubmixEffectChain)
 			{
 				if (UResonanceAudioReverbPluginPreset* PluginPreset = Cast<UResonanceAudioReverbPluginPreset>(Preset))
 				{
@@ -229,15 +250,18 @@ namespace ResonanceAudio
 
 			if (!bFoundPreset)
 			{
-				static const FString DefaultPresetName = TEXT("ResonanceReverbDefault_0");
-				UE_LOG(LogResonanceAudio, Error, TEXT("Failed to find Resonance UResonanceAudioReverbPluginPreset on default reverb submix. Creating stub '%s'."),
-					*Settings->OutputSubmix.GetAssetPathString(),
-					*DefaultPresetName);
-				ReverbSubmix->SubmixEffectChain.Add(ReverbPluginPreset);
+				if (!IgnoreUserResonanceSubmixCVar)
+				{
+					static const FString DefaultPresetName = TEXT("ResonanceReverbDefault_0");
+					UE_LOG(LogResonanceAudio, Error, TEXT("Failed to find Resonance UResonanceAudioReverbPluginPreset on default reverb submix. Creating stub '%s'."),
+						*Settings->OutputSubmix.GetAssetPathString(),
+						*DefaultPresetName);
+				}
+				ReverbSubmixPtr->SubmixEffectChain.Add(ReverbPluginPreset);
 			}
 		}
 
-		return ReverbSubmix;
+		return ReverbSubmixPtr;
 	}
 
 	void FResonanceAudioReverb::SetPreset(UResonanceAudioReverbPluginPreset* InPreset)
diff --git a/Engine/Plugins/Runtime/ResonanceAudio/Source/ResonanceAudio/Private/ResonanceAudioReverb.h b/Engine/Plugins/Runtime/ResonanceAudio/Source/ResonanceAudio/Private/ResonanceAudioReverb.h
--- a/Engine/Plugins/Runtime/ResonanceAudio/Source/ResonanceAudio/Private/ResonanceAudioReverb.h
+++ b/Engine/Plugins/Runtime/ResonanceAudio/Source/ResonanceAudio/Private/ResonanceAudioReverb.h
@@ -151,8 +151,8 @@ namespace ResonanceAudio
 
 		FSoundEffectSubmixPtr SubmixEffect;
 
-		// Stub Reverb submix to be used only when Resonance fails to load the submix asset.
-		USoundSubmix* StubSubmixPtr;
+		// Reverb submix to be used only when Resonance fails to load the submix asset.
+		USoundSubmix* ReverbSubmixPtr;
 	};
 
 } // namespace ResonanceAudio
diff --git a/Engine/Source/Runtime/Engine/Classes/Sound/SoundSubmix.h b/Engine/Source/Runtime/Engine/Classes/Sound/SoundSubmix.h
--- a/Engine/Source/Runtime/Engine/Classes/Sound/SoundSubmix.h
+++ b/Engine/Source/Runtime/Engine/Classes/Sound/SoundSubmix.h
@@ -255,11 +255,9 @@ public:
 
 protected:
 
-#if WITH_EDITOR
 	ENGINE_API virtual void PostDuplicate(EDuplicateMode::Type DuplicateMode) override;
 	ENGINE_API virtual void PreEditChange(FProperty* PropertyAboutToChange) override;
 	ENGINE_API virtual void PostEditChangeProperty(struct FPropertyChangedEvent& PropertyChangedEvent) override;
-#endif // WITH_EDITOR
 	//~ End UObject Interface.
 
 private:
@@ -293,9 +291,10 @@ public:
 	/**
 	* Set the parent submix of this SoundSubmix, removing it as a child from its previous owner
 	*
-	* @param	InParentSubmix	The New Parent Submix of this
+	* @param	InParentSubmix	The New Parent Submix of this submix
+	* @param	bModifyAssets	Whether or not to mark the UObjects as modified.
 	*/
-	ENGINE_API void SetParentSubmix(USoundSubmixBase* InParentSubmix);
+	ENGINE_API void SetParentSubmix(USoundSubmixBase* InParentSubmix, bool bModifyAssets = true);
 
 	ENGINE_API virtual bool DynamicConnect(const UObject* WorldContextObject, USoundSubmixBase* Parent) override;
 	ENGINE_API bool DynamicConnect(FAudioDeviceHandle Handle, USoundSubmixBase* Parent);
diff --git a/Engine/Source/Runtime/Engine/Private/SoundSubmix.cpp b/Engine/Source/Runtime/Engine/Private/SoundSubmix.cpp
--- a/Engine/Source/Runtime/Engine/Private/SoundSubmix.cpp
+++ b/Engine/Source/Runtime/Engine/Private/SoundSubmix.cpp
@@ -875,18 +875,26 @@ void USoundSubmixBase::PostEditChangeProperty(FPropertyChangedEvent& PropertyCha
 }
 
 TArray<TObjectPtr<USoundSubmixBase>> USoundSubmixBase::BackupChildSubmixes;
+#endif
+
 
-void USoundSubmixWithParentBase::SetParentSubmix(USoundSubmixBase* InParentSubmix)
+void USoundSubmixWithParentBase::SetParentSubmix(USoundSubmixBase* InParentSubmix, bool bModifyAssets)
 {
 	if (ParentSubmix != InParentSubmix)
 	{
 		if (ParentSubmix)
 		{
-			ParentSubmix->Modify();
+			if (bModifyAssets)
+			{
+				ParentSubmix->Modify();
+			}
 			ParentSubmix->ChildSubmixes.Remove(this);
 		}
 
-		Modify();
+		if (bModifyAssets)
+		{
+			Modify();
+		}
 		ParentSubmix = InParentSubmix;
 		if (ParentSubmix)
 		{
@@ -953,7 +961,6 @@ void USoundSubmixWithParentBase::PostDuplicate(EDuplicateMode::Type DuplicateMod
 
 	Super::PostDuplicate(DuplicateMode);
 }
-#endif
 
 void USoundSubmixBase::AddReferencedObjects(UObject* InThis, FReferenceCollector& Collector)
 {
