Use PRAGMAs for unreachable code warnings where appropriate and put in clang disabling as needed
Fix some unreachable code warnings

#rnx

diff --git a/Engine/Plugins/Runtime/ResonanceAudio/Source/ResonanceAudio/Private/ResonanceAudioLibrary/resonance_audio/utils/lockless_task_queue.cc b/Engine/Plugins/Runtime/ResonanceAudio/Source/ResonanceAudio/Private/ResonanceAudioLibrary/resonance_audio/utils/lockless_task_queue.cc
--- a/Engine/Plugins/Runtime/ResonanceAudio/Source/ResonanceAudio/Private/ResonanceAudioLibrary/resonance_audio/utils/lockless_task_queue.cc
+++ b/Engine/Plugins/Runtime/ResonanceAudio/Source/ResonanceAudio/Private/ResonanceAudioLibrary/resonance_audio/utils/lockless_task_queue.cc
@@ -17,6 +17,7 @@ limitations under the License.
 #if defined(__clang__)
 _Pragma("clang diagnostic push") \
 _Pragma("clang diagnostic ignored \"-Wshadow\"")
+_Pragma("clang diagnostic ignored \"-Wunreachable-code-aggressive\"")
 #elif defined(_MSC_VER)
 #pragma warning(push)
 #pragma warning(disable:6294) /* Ill-defined for-loop:  initial condition does not satisfy test.  Loop body not executed. */
@@ -27,7 +28,7 @@ _Pragma("clang diagnostic ignored \"-Wshadow\"")
 #pragma warning(disable:4459) /* declaration of 'LocalVariable' hides global declaration */ 
 #pragma warning(disable:6244) /* local declaration of <variable> hides previous declaration at <line> of <file> */
 #pragma warning(disable:4702) /* unreachable code */
-#pragma warning(disable:4100) /* unreachable code */
+#pragma warning(disable:4100) /* unreferenced formal parameter */
 #endif
 
 #include "utils/lockless_task_queue.h"
diff --git a/Engine/Plugins/Runtime/ResonanceAudio/Source/ResonanceAudio/Private/ResonanceAudioLibrary/resonance_audio/utils/lockless_task_queue.h b/Engine/Plugins/Runtime/ResonanceAudio/Source/ResonanceAudio/Private/ResonanceAudioLibrary/resonance_audio/utils/lockless_task_queue.h
--- a/Engine/Plugins/Runtime/ResonanceAudio/Source/ResonanceAudio/Private/ResonanceAudioLibrary/resonance_audio/utils/lockless_task_queue.h
+++ b/Engine/Plugins/Runtime/ResonanceAudio/Source/ResonanceAudio/Private/ResonanceAudioLibrary/resonance_audio/utils/lockless_task_queue.h
@@ -20,6 +20,7 @@ limitations under the License.
 #if defined(__clang__)
 _Pragma("clang diagnostic push") \
 _Pragma("clang diagnostic ignored \"-Wshadow\"")
+_Pragma("clang diagnostic ignored \"-Wunreachable-code-aggressive\"")
 #elif defined(_MSC_VER)
 #pragma warning(push)
 #pragma warning(disable:6294) /* Ill-defined for-loop:  initial condition does not satisfy test.  Loop body not executed. */
@@ -30,7 +31,7 @@ _Pragma("clang diagnostic ignored \"-Wshadow\"")
 #pragma warning(disable:4459) /* declaration of 'LocalVariable' hides global declaration */ 
 #pragma warning(disable:6244) /* local declaration of <variable> hides previous declaration at <line> of <file> */
 #pragma warning(disable:4702) /* unreachable code */
-#pragma warning(disable:4100) /* unreachable code */
+#pragma warning(disable:4100) /* unreferenced formal parameter */
 #endif
 
 #include <stddef.h>
diff --git a/Engine/Plugins/Runtime/ResonanceAudio/Source/ResonanceAudio/Private/ResonanceAudioLibrary/resonance_audio/utils/task_thread_pool.cc b/Engine/Plugins/Runtime/ResonanceAudio/Source/ResonanceAudio/Private/ResonanceAudioLibrary/resonance_audio/utils/task_thread_pool.cc
--- a/Engine/Plugins/Runtime/ResonanceAudio/Source/ResonanceAudio/Private/ResonanceAudioLibrary/resonance_audio/utils/task_thread_pool.cc
+++ b/Engine/Plugins/Runtime/ResonanceAudio/Source/ResonanceAudio/Private/ResonanceAudioLibrary/resonance_audio/utils/task_thread_pool.cc
@@ -17,6 +17,7 @@ limitations under the License.
 #if defined(__clang__)
 _Pragma("clang diagnostic push") \
 _Pragma("clang diagnostic ignored \"-Wshadow\"")
+_Pragma("clang diagnostic ignored \"-Wunreachable-code-aggressive\"")
 #elif defined(_MSC_VER)
 #pragma warning(push)
 #pragma warning(disable:6294) /* Ill-defined for-loop:  initial condition does not satisfy test.  Loop body not executed. */
diff --git a/Engine/Plugins/Runtime/ResonanceAudio/Source/ResonanceAudio/Private/ResonanceAudioLibrary/resonance_audio/utils/task_thread_pool.h b/Engine/Plugins/Runtime/ResonanceAudio/Source/ResonanceAudio/Private/ResonanceAudioLibrary/resonance_audio/utils/task_thread_pool.h
--- a/Engine/Plugins/Runtime/ResonanceAudio/Source/ResonanceAudio/Private/ResonanceAudioLibrary/resonance_audio/utils/task_thread_pool.h
+++ b/Engine/Plugins/Runtime/ResonanceAudio/Source/ResonanceAudio/Private/ResonanceAudioLibrary/resonance_audio/utils/task_thread_pool.h
@@ -20,6 +20,7 @@ limitations under the License.
 #if defined(__clang__)
 _Pragma("clang diagnostic push") \
 _Pragma("clang diagnostic ignored \"-Wshadow\"")
+_Pragma("clang diagnostic ignored \"-Wunreachable-code-aggressive\"")
 #elif defined(_MSC_VER)
 #pragma warning(push)
 #pragma warning(disable:6294) /* Ill-defined for-loop:  initial condition does not satisfy test.  Loop body not executed. */
diff --git a/Engine/Plugins/Runtime/ResonanceAudio/Source/ResonanceAudio/Private/ResonanceEigenHelper.h b/Engine/Plugins/Runtime/ResonanceAudio/Source/ResonanceAudio/Private/ResonanceEigenHelper.h
--- a/Engine/Plugins/Runtime/ResonanceAudio/Source/ResonanceAudio/Private/ResonanceEigenHelper.h
+++ b/Engine/Plugins/Runtime/ResonanceAudio/Source/ResonanceAudio/Private/ResonanceEigenHelper.h
@@ -9,6 +9,7 @@
 #if defined(__clang__)
 _Pragma("clang diagnostic push") \
 _Pragma("clang diagnostic ignored \"-Wshadow\"")
+_Pragma("clang diagnostic ignored \"-Wunreachable-code-aggressive\"")
 #elif defined(_MSC_VER)
 #pragma warning(push)
 #pragma warning(disable:6294) /* Ill-defined for-loop:  initial condition does not satisfy test.  Loop body not executed. */
